Moved ChaosPhysicsTests out of into Engine/Source/Programs. This is prep work to get the tests runnning in horde.

#Synced-CL 51835286

diff --git a/Engine/Source/Programs/ChaosPhysicsTests/ChaosPhysicsTests.Build.cs b/Engine/Source/Programs/ChaosPhysicsTests/ChaosPhysicsTests.Build.cs
new file mode 100644
--- /dev/null
+++ b/Engine/Source/Programs/ChaosPhysicsTests/ChaosPhysicsTests.Build.cs
@@ -0,0 +1,60 @@
+// Copyright Epic Games, Inc. All Rights Reserved.
+
+using UnrealBuildTool;
+
+public class ChaosPhysicsTests : TestModuleRules
+{
+	static ChaosPhysicsTests()
+	{
+		if (InTestMode)
+		{
+			TestMetadata = new Metadata();
+			TestMetadata.TestName = "ChaosPhysicsTests";
+			TestMetadata.TestShortName = "ChaosPhysicsTests";
+			TestMetadata.ReportType = "testdata";
+			TestMetadata.SupportedPlatforms.Add(UnrealTargetPlatform.Linux);
+			TestMetadata.SupportedPlatforms.Add(UnrealTargetPlatform.Mac);
+			TestMetadata.SupportedPlatforms.Add(UnrealTargetPlatform.Android);
+			TestMetadata.SupportedPlatforms.Add(UnrealTargetPlatform.IOS);
+
+			string PlatformCompilationArgs;
+			foreach (var Platform in UnrealTargetPlatform.GetValidPlatforms())
+			{
+				if (Platform == UnrealTargetPlatform.Android)
+				{
+					PlatformCompilationArgs = "-allmodules -architectures=arm64";
+				}
+				else
+				{
+					PlatformCompilationArgs = "-allmodules";
+				}
+				TestMetadata.PlatformCompilationExtraArgs.Add(Platform, PlatformCompilationArgs);
+			}
+
+			// Platform-specific tags
+			TestMetadata.PlatformTags.Add(UnrealTargetPlatform.Linux, "~[.]~[Slow]");
+			TestMetadata.PlatformTags.Add(UnrealTargetPlatform.Android, "~[Perf]~[Slow]~[AndroidSkip]");
+
+			// Allow Android run for this test
+			// Will remove Android from PlatformsRunUnsupported as more diverse types of tests can run on this platform
+			TestMetadata.PlatformsRunUnsupported.Remove(UnrealTargetPlatform.Android);
+		}
+	}
+
+	public ChaosPhysicsTests(ReadOnlyTargetRules Target) : base(Target)
+	{
+		PrivateDependencyModuleNames.AddRange(
+			new string[] {
+				"Core",
+				"CoreUObject",
+				"Engine",
+				"ChaosTestHarness",
+		});
+
+		// Temporarily link to base Chaos modules - remove once everything is moved into the new tree
+		SetupModulePhysicsSupport(Target);
+
+		bTreatAsEngineModule = true;
+		bWarningsAsErrors = true;
+	}
+}
\ No newline at end of file
diff --git a/Engine/Source/Programs/ChaosPhysicsTests/ChaosPhysicsTests.Target.cs b/Engine/Source/Programs/ChaosPhysicsTests/ChaosPhysicsTests.Target.cs
new file mode 100644
--- /dev/null
+++ b/Engine/Source/Programs/ChaosPhysicsTests/ChaosPhysicsTests.Target.cs
@@ -0,0 +1,12 @@
+// Copyright Epic Games, Inc. All Rights Reserved.
+
+using UnrealBuildTool;
+
+[SupportedPlatforms(UnrealPlatformClass.All)]
+ public class ChaosPhysicsTestsTarget : ChaosTestHarnessTarget
+ {
+     public ChaosPhysicsTestsTarget(TargetInfo Target) : base(Target)
+     {
+		TestTargetRules.bTestsRequireEngine = true;
+     }
+ }
\ No newline at end of file
diff --git a/Engine/Source/Programs/ChaosPhysicsTests/Private/AccelerationStructureHandleTests.cpp b/Engine/Source/Programs/ChaosPhysicsTests/Private/AccelerationStructureHandleTests.cpp
new file mode 100644
--- /dev/null
+++ b/Engine/Source/Programs/ChaosPhysicsTests/Private/AccelerationStructureHandleTests.cpp
@@ -0,0 +1,151 @@
+// Copyright Epic Games, Inc. All Rights Reserved.
+
+#include "CoreMinimal.h"
+#include "ChaosTestHarness.h"
+#include "ChaosTestScene.h"
+
+#include "Chaos/Box.h"
+#include "PhysicsProxy/SingleParticlePhysicsProxy.h"
+
+namespace Chaos
+{
+	void SetShapeFilter(FSingleParticlePhysicsProxy* Proxy, const int32 ShapeIndex, uint8 CollisionChannelIndex, Chaos::EFilterFlags Flags, uint32 BlockMask, uint32 OverlapMask)
+	{
+		Filter::FShapeFilterBuilder Builder;
+		Builder.SetCollisionChannelIndex(CollisionChannelIndex);
+		Builder.SetBlockChannelMask(BlockMask);
+		Builder.SetOverlapChannelMask(OverlapMask);
+		Builder.SetFilterFlags(Flags);
+		const Filter::FShapeFilterData FilterData = Builder.Build();
+
+		Proxy->GetGameThreadAPI().SetShapeFilterData(ShapeIndex, FilterData);
+	}
+
+	void SetShapeEnabled(FSingleParticlePhysicsProxy* Proxy, const int32 ShapeIndex, bool bSimEnabled, bool bQueryEnabled)
+	{
+		FRigidBodyHandle_External& BodyHandle = Proxy->GetGameThreadAPI();
+		BodyHandle.SetShapeSimCollisionEnabled(ShapeIndex, bSimEnabled);
+		BodyHandle.SetShapeQueryCollisionEnabled(ShapeIndex, bQueryEnabled);
+	}
+
+	bool PrePreQueryFilter(FSingleParticlePhysicsProxy* Proxy, uint8 CollisionChannelIndex, const uint32 BlockMask)
+	{
+		Chaos::Filter::FQueryTraceFilterBuilder QueryBuilder;
+		QueryBuilder.SetCollisionChannelIndex(CollisionChannelIndex);
+		QueryBuilder.SetBlockChannelMask(BlockMask);
+		const Chaos::Filter::FQueryFilterData QueryFilter = QueryBuilder.Build();
+
+		FAccelerationStructureHandle ParticleHandle(Proxy->GetParticle_LowLevel());
+		return ParticleHandle.PrePreFilter(QueryFilter);
+	}
+
+	bool PrePreSimFilter(FSingleParticlePhysicsProxy* Proxy, uint8 CollisionChannelIndex, const uint32 BlockMask)
+	{
+		Chaos::Filter::FShapeFilterBuilder Builder;
+		Builder.SetCollisionChannelIndex(CollisionChannelIndex);
+		Builder.SetBlockChannelMask(BlockMask);
+
+		Chaos::Filter::FShapeUnionFilterData ParticleFilter;
+		ParticleFilter.Combine(Builder.Build());
+
+		FAccelerationStructureHandle ParticleHandle(Proxy->GetParticle_LowLevel());
+		return ParticleHandle.PrePreFilter(ParticleFilter);
+	}
+
+	TEST_CASE("FAccelerationStructureHandle - PrePreFilter", "[Chaos][unit]")
+	{
+		using namespace Chaos::LowLevelTest;
+
+		FChaosTestScene Scene;
+
+		const FVector3d HalfExtents(50, 50, 50);
+		TArray<FImplicitObjectPtr> Geometry
+		{
+			FImplicitObjectPtr(new Chaos::FImplicitBox3(-HalfExtents, HalfExtents)),
+			FImplicitObjectPtr(new Chaos::FImplicitBox3(-HalfExtents, HalfExtents)),
+			FImplicitObjectPtr(new Chaos::FImplicitBox3(-HalfExtents, HalfExtents)),
+			FImplicitObjectPtr(new Chaos::FImplicitBox3(-HalfExtents, HalfExtents)),
+		};
+		FImplicitObjectPtr UnionImplicit = FImplicitObjectPtr(new FImplicitObjectUnion(MoveTemp(Geometry)));
+		FSingleParticlePhysicsProxy* UnionProxy0 = Scene.AddDynamicProxy(TEXT("Union0"), FTransform3d(FVector3d(0, 0, 0)), UnionImplicit);		
+		
+		for (int32 ShapeIndex = 0; ShapeIndex < 4; ++ShapeIndex)
+		{
+			const int32 CollisionChannelIndex = ShapeIndex;
+			// Setup the block and overlap masks with a gap in-between.
+			const int32 BlockMask = 1 << ShapeIndex;
+			const int32 OverlapMask = 1 << (ShapeIndex + 5);
+			SetShapeFilter(UnionProxy0, ShapeIndex, CollisionChannelIndex, EFilterFlags::All, BlockMask, OverlapMask);
+		}
+
+		SECTION("Basic")
+		{
+			// Sim: Channels(0b1111) Block(0b1111)
+			// Query: Channels(0b1111) Block(0b1111) Overlap(0b111100000)
+			
+			// Query only uses the channel index. Test all values.
+			CHECK(false == PrePreQueryFilter(UnionProxy0, 0, 0));
+			CHECK(false == PrePreQueryFilter(UnionProxy0, 1, 0));
+			CHECK(false == PrePreQueryFilter(UnionProxy0, 2, 0));
+			CHECK(false == PrePreQueryFilter(UnionProxy0, 3, 0));
+			CHECK(true == PrePreQueryFilter(UnionProxy0, 4, 0));
+			CHECK(false == PrePreQueryFilter(UnionProxy0, 5, 0));
+			CHECK(false == PrePreQueryFilter(UnionProxy0, 6, 0));
+			CHECK(false == PrePreQueryFilter(UnionProxy0, 7, 0));
+			CHECK(false == PrePreQueryFilter(UnionProxy0, 8, 0));
+			CHECK(true == PrePreQueryFilter(UnionProxy0, 9, 0));
+
+			// Sim uses both.
+			// First test all valid collision channels with a block mask that passes all shapes.
+			CHECK(false == PrePreSimFilter(UnionProxy0, 0, 0b1111));
+			CHECK(false == PrePreSimFilter(UnionProxy0, 1, 0b1111));
+			CHECK(false == PrePreSimFilter(UnionProxy0, 2, 0b1111));
+			CHECK(false == PrePreSimFilter(UnionProxy0, 3, 0b1111));
+			// Now test all valid block masks with a collision channel index that passes all.
+			CHECK(false == PrePreSimFilter(UnionProxy0, 0, 0b0001));
+			CHECK(false == PrePreSimFilter(UnionProxy0, 0, 0b0010));
+			CHECK(false == PrePreSimFilter(UnionProxy0, 0, 0b0100));
+			CHECK(false == PrePreSimFilter(UnionProxy0, 0, 0b1000));
+
+			// Now test edge cases that will be filtered.
+			CHECK(true == PrePreSimFilter(UnionProxy0, 0, 0b0000));
+			CHECK(true == PrePreSimFilter(UnionProxy0, 0, 0b10000));
+			CHECK(true == PrePreSimFilter(UnionProxy0, 4, 0b1111));
+		}
+		SECTION("Mixed Sim/Query Enabled")
+		{
+			SetShapeEnabled(UnionProxy0, 0, false, false);
+			SetShapeEnabled(UnionProxy0, 1, false, true);
+			SetShapeEnabled(UnionProxy0, 2, true, false);
+			SetShapeEnabled(UnionProxy0, 3, true, true);
+			// Sim: Channels(0b1100) Block(0b1100)
+			// Query: Channels(0b1010) Block(0b1010) Overlap(0b101000000) Combined(0b101001010)
+
+			// Test all query channel indices.
+			CHECK(true == PrePreQueryFilter(UnionProxy0, 0, 0));
+			CHECK(false == PrePreQueryFilter(UnionProxy0, 1, 0));
+			CHECK(true == PrePreQueryFilter(UnionProxy0, 2, 0));
+			CHECK(false == PrePreQueryFilter(UnionProxy0, 3, 0));
+			CHECK(true == PrePreQueryFilter(UnionProxy0, 4, 0));
+			CHECK(true == PrePreQueryFilter(UnionProxy0, 5, 0));
+			CHECK(false == PrePreQueryFilter(UnionProxy0, 6, 0));
+			CHECK(true == PrePreQueryFilter(UnionProxy0, 7, 0));
+			CHECK(false == PrePreQueryFilter(UnionProxy0, 8, 0));
+			CHECK(true == PrePreQueryFilter(UnionProxy0, 9, 0));
+
+			// Test all channel indices.
+			CHECK(true == PrePreSimFilter(UnionProxy0, 0, 0b1111));
+			CHECK(true == PrePreSimFilter(UnionProxy0, 1, 0b1111));
+			CHECK(false == PrePreSimFilter(UnionProxy0, 2, 0b1111));
+			CHECK(false == PrePreSimFilter(UnionProxy0, 3, 0b1111));
+			CHECK(true == PrePreSimFilter(UnionProxy0, 4, 0b1111));
+
+			// Test all block masks.
+			CHECK(true == PrePreSimFilter(UnionProxy0, 2, 0b0001));
+			CHECK(true == PrePreSimFilter(UnionProxy0, 2, 0b0010));
+			CHECK(false == PrePreSimFilter(UnionProxy0, 2, 0b0100));
+			CHECK(false == PrePreSimFilter(UnionProxy0, 2, 0b1000));
+			CHECK(true == PrePreSimFilter(UnionProxy0, 2, 0b10000));
+		}
+	}
+} // namespace Chaos
diff --git a/Engine/Source/Programs/ChaosPhysicsTests/Private/BasicTests.cpp b/Engine/Source/Programs/ChaosPhysicsTests/Private/BasicTests.cpp
new file mode 100644
--- /dev/null
+++ b/Engine/Source/Programs/ChaosPhysicsTests/Private/BasicTests.cpp
@@ -0,0 +1,105 @@
+// Copyright Epic Games, Inc. All Rights Reserved.
+
+#include "CoreMinimal.h"
+#include "ChaosTestHarness.h"
+
+#include "Chaos/AABB.h"
+
+// Simple test cases
+TEST_CASE("AABB Construction", "[Chaos][unit]")
+{
+	using namespace Chaos;
+	FVector HalfExtent{ 100.0f, 100.0f, 100.0f };
+	FAABB3 DefaultAabb;
+	FAABB3 InitializedAabb(-HalfExtent, HalfExtent);
+
+	REQUIRE(DefaultAabb.IsEmpty());
+	REQUIRE(!InitializedAabb.IsEmpty());
+	REQUIRE(InitializedAabb.GetVolume() == 8000000);
+}
+
+TEST_CASE("AABB Intersection", "[Chaos][unit]")
+{
+	using namespace Chaos;
+
+	FAABB3 A{ {-100.0f, -100.0f, -100.0f}, {10.0f, 10.0f, 10.0f} };
+	FAABB3 B{ {-10.0f, -10.0f, -10.0f}, {100.0f, 100.0f, 100.0f} };
+	FAABB3 Intersect = A.GetIntersection(B);
+
+	REQUIRE(Intersect.Min() == FVector{ -10.0f, -10.0f, -10.0f });
+	REQUIRE(Intersect.Max() == FVector{ 10.0f, 10.0f, 10.0f });
+}
+
+// Test with Fixtures
+TEST_CASE("AABB Operations", "[Chaos][unit]")
+{
+	// Rather than use a fixture class, this portion runs for
+	// all sections below, each section can mutate and use the 
+	// state without interfering with each other
+	using namespace Chaos;
+	FVector HalfExtent{ 100.0f, 100.0f, 100.0f };
+	FAABB3 TestAabb(-HalfExtent, HalfExtent);
+	const float InitialVolume = TestAabb.GetVolume();
+
+	SECTION("Growing around another box expands")
+	{
+		TestAabb.GrowToInclude({{-200, -200, -200}, {200, 200, 200}});
+		const float NewVolume = TestAabb.GetVolume();
+		REQUIRE(NewVolume > InitialVolume);
+	}
+
+	SECTION("Growing around a fully contained box does not expand")
+	{
+		TestAabb.GrowToInclude({ {-20, -20, -20}, {20, 20, 20} });
+		const float NewVolume = TestAabb.GetVolume();
+		REQUIRE(NewVolume == InitialVolume);
+	}
+}
+
+// BDD Style test matching above
+SCENARIO("AABBs can grow to include larger AABBs but not grow if an enclosed AABB is provided", "[Chaos][unit]")
+{
+	GIVEN("A 200x200x200 AABB")
+	{
+		using namespace Chaos;
+		FVector HalfExtent{ 100.0f, 100.0f, 100.0f };
+		FAABB3 TestAabb(-HalfExtent, HalfExtent);
+		const float InitialVolume = TestAabb.GetVolume();
+
+		REQUIRE(TestAabb.GetVolume() == 8000000);
+
+		WHEN("It grows around a larger AABB")
+		{
+			TestAabb.GrowToInclude({ {-200, -200, -200}, {200, 200, 200} });
+
+			THEN("Its volume increases")
+			{
+				REQUIRE(TestAabb.GetVolume() > InitialVolume);
+			}
+		}
+
+		WHEN("It grows around a smaller AABB")
+		{
+			TestAabb.GrowToInclude({ {-20, -20, -20}, {20, 20, 20} });
+
+			THEN("Its volume remains the same")
+			{
+				REQUIRE(TestAabb.GetVolume() == InitialVolume);
+			}
+		}
+	}
+}
+
+// Benchmark
+TEST_CASE("AABB Transform Benchmark", "[Chaos][!benchmark]")
+{
+	using namespace Chaos;
+
+	FTransform Transform(FRotator{ 90, 0, 0 }, FVector{ 0, 0, 0 });
+	FAABB3 Aabb{ {100, 100, 100}, {200, 200, 200} };
+
+	BENCHMARK("Rotate X 90")
+	{
+		return Aabb.TransformedAABB(Transform);
+	};
+}
\ No newline at end of file
diff --git a/Engine/Source/Programs/ChaosPhysicsTests/Private/ChaosTestHarness.h b/Engine/Source/Programs/ChaosPhysicsTests/Private/ChaosTestHarness.h
new file mode 100644
--- /dev/null
+++ b/Engine/Source/Programs/ChaosPhysicsTests/Private/ChaosTestHarness.h
@@ -0,0 +1,8 @@
+// Copyright Epic Games, Inc. All Rights Reserved.
+
+#pragma once
+
+// From LevelLevelTests
+#include "TestHarness.h"
+
+#include "Chaos/LowLevelTest/ChaosTestHarness.h"
diff --git a/Engine/Source/Programs/ChaosPhysicsTests/Private/ChaosTestScene.h b/Engine/Source/Programs/ChaosPhysicsTests/Private/ChaosTestScene.h
new file mode 100644
--- /dev/null
+++ b/Engine/Source/Programs/ChaosPhysicsTests/Private/ChaosTestScene.h
@@ -0,0 +1,4 @@
+// Copyright Epic Games, Inc. All Rights Reserved.
+#pragma once
+
+#include "Chaos/LowLevelTest/ChaosTestScene.h"
diff --git a/Engine/Source/Programs/ChaosPhysicsTests/Private/CollisionDataTests.cpp b/Engine/Source/Programs/ChaosPhysicsTests/Private/CollisionDataTests.cpp
new file mode 100644
--- /dev/null
+++ b/Engine/Source/Programs/ChaosPhysicsTests/Private/CollisionDataTests.cpp
@@ -0,0 +1,517 @@
+// Copyright Epic Games, Inc. All Rights Reserved.
+#include "CoreMinimal.h"
+
+#include "ChaosTestHarness.h"
+#include <catch2/generators/catch_generators_range.hpp>
+
+#include "Chaos/ParticleDirtyFlags.h"
+#include "Chaos/ChaosArchive.h"
+#include "Serialization/MemoryReader.h"
+#include "Serialization/MemoryWriter.h"
+#include "UObject/FortniteMainBranchObjectVersion.h"
+
+namespace Chaos::Tests
+{
+	using FMaskFilter = uint8;
+
+	// Helper to output the serialization data to a code string so we can easily hardcode the test
+	FString ToString(const TArray<uint8>& DataArray)
+	{
+		TStringBuilder<128> Builder;
+		Builder.Append("{");
+		for (const uint8 Data : DataArray)
+		{
+			if (Data != 0)
+			{
+				Builder.Appendf(TEXT("0x%X, "), Data);
+			}
+			else
+			{
+				Builder.Appendf(TEXT("%X, "), Data);
+			}
+		}
+		Builder.Append("}");
+		return Builder.ToString();
+	}
+
+	void Save(FCollisionData& InCollisionData, const int32 Version, TArray<uint8>& OutData)
+	{
+		FMemoryWriter Ar(OutData);
+		FChaosArchive Writer(Ar);
+		Writer.SetCustomVersion(FFortniteMainBranchObjectVersion::GUID, Version, FName("Unused"));
+		Writer.SetShouldSkipUpdateCustomVersion(true);
+		Writer << InCollisionData;
+	}
+
+	void Load(const TArray<uint8>& InData, const int32 Version, FCollisionData& OutCollisionData)
+	{
+		FMemoryReader Ar(InData);
+		FChaosArchive Reader(Ar);
+		Reader.SetCustomVersion(FFortniteMainBranchObjectVersion::GUID, Version, FName("Unused"));
+		Reader << OutCollisionData;
+	}
+
+	uint32 BuildLegacyWord3(const FMaskFilter MaskFilter, const uint8 CollisionIndex, const EFilterFlags Flags)
+	{
+		constexpr uint32 NumBitsFlags = 21;
+		constexpr uint32 NumBitsCollisionChannel = 5;
+		constexpr uint32 NumBitsMaskFilter = 6;
+		const uint32 FlagsMask = (1 << NumBitsFlags) - 1;
+		const uint32 CollisionChannelMask = (1 << NumBitsCollisionChannel) - 1;
+		const uint32 MaskFilterMask = (1 << NumBitsMaskFilter) - 1;
+
+		uint32 Result = 0;
+		Result |= (uint32)Flags & FlagsMask;
+		Result |= ((uint32)CollisionIndex & CollisionChannelMask) << NumBitsFlags;
+		Result |= ((uint32)MaskFilter & MaskFilterMask) << (NumBitsFlags + NumBitsCollisionChannel);
+		return Result;
+	}
+
+	FCollisionFilterData BuildLegacyQueryFilter(const uint32 OwnerId, const uint32 BlockMask, const uint32 OverlapMask, const FMaskFilter MaskFilter, const uint8 CollisionIndex, const EFilterFlags Flags)
+	{
+		FCollisionFilterData Result;
+		Result.Word0 = OwnerId;
+		Result.Word1 = BlockMask;
+		Result.Word2 = OverlapMask;
+		Result.Word3 = BuildLegacyWord3(MaskFilter, CollisionIndex, Flags);
+		return Result;
+	}
+
+	FCollisionFilterData BuildLegacySimFilter(const uint32 BodyIndex, const uint32 BlockMask, const uint32 ComponentId, const FMaskFilter MaskFilter, const uint8 CollisionIndex, const EFilterFlags Flags)
+	{
+		FCollisionFilterData Result;
+		Result.Word0 = BodyIndex;
+		Result.Word1 = BlockMask;
+		Result.Word2 = ComponentId;
+		Result.Word3 = BuildLegacyWord3(MaskFilter, CollisionIndex, Flags);
+		return Result;
+	}
+
+	struct FCombinedTestData
+	{
+		uint32 BodyIndex = 0;
+		uint32 ComponentId = 0;
+		uint32 OwnerId = 0;
+		uint64 OverlapMask = 0;
+		uint64 BlockMask = 0;
+		FMaskFilter MaskFilter = 0;
+		uint8 CollisionChannel = 0;
+		EFilterFlags Flags = EFilterFlags::None;
+
+		static_assert(PLATFORM_LITTLE_ENDIAN, "Serialization needs updating to support big-endian platforms!");
+		TArray<uint8> SerializationData;
+
+		void SaveLegacy(const int32 Version)
+		{
+			FCollisionData CollisionData;
+			PRAGMA_DISABLE_DEPRECATION_WARNINGS
+			CollisionData.SimData = BuildLegacySimFilter(BodyIndex, (uint32)BlockMask, ComponentId, MaskFilter, CollisionChannel, Flags);
+			CollisionData.QueryData = BuildLegacyQueryFilter(OwnerId, (uint32)BlockMask, (uint32)OverlapMask, MaskFilter, CollisionChannel, Flags);
+			PRAGMA_ENABLE_DEPRECATION_WARNINGS
+
+			Chaos::Tests::Save(CollisionData, Version, SerializationData);
+		}
+
+		void Save(const int32 Version)
+		{
+			Chaos::Filter::FShapeFilterBuilder Builder;
+			Builder.SetOverlapChannelMask(OverlapMask);
+			Builder.SetBlockChannelMask(BlockMask);
+			Builder.SetCollisionChannelIndex(CollisionChannel);
+			Builder.SetMaskFilter(MaskFilter);
+			Builder.SetFilterFlags(Flags);
+
+			FCollisionData CollisionData;
+			CollisionData.SetFilterInstanceData(Chaos::Filter::FInstanceData(OwnerId, ComponentId));
+			CollisionData.SetShapeFilterData(Builder.Build());
+
+			Chaos::Tests::Save(CollisionData, Version, SerializationData);
+		}
+
+		FCollisionData Load(const int32 Version) const
+		{
+			FCollisionData ResultCollisionData;
+			Chaos::Tests::Load(SerializationData, Version, ResultCollisionData);
+			return ResultCollisionData;
+		}
+	};
+
+	struct FLegacyFilterTestData
+	{
+		uint32 BodyIndex = 0;
+		uint32 ComponentId = 0;
+		uint32 OwnerId = 0;
+
+		uint32 QueryOverlapMask = 0;
+		uint32 QueryBlockMask = 0;
+		uint32 SimBlockMask = 0;
+		uint32 ExpectedBlockMask = 0;
+
+		FMaskFilter QueryMaskFilter = 0;
+		FMaskFilter SimMaskFilter = 0;
+		FMaskFilter ExpectedMaskFilter = 0;
+
+		uint8 QueryCollisionChannel = 0;
+		uint8 SimCollisionChannel = 0;
+		uint8 ExpectedCollisionChannel = 0;
+
+		EFilterFlags QueryFlags = EFilterFlags::None;
+		EFilterFlags SimFlags = EFilterFlags::None;
+		EFilterFlags ExpectedFlags = EFilterFlags::None;
+
+		static_assert(PLATFORM_LITTLE_ENDIAN, "Serialization needs updating to support big-endian platforms!");
+		TArray<uint8> SerializationData;
+
+		void Save(const int32 Version)
+		{
+			FCollisionData CollisionData;
+			PRAGMA_DISABLE_DEPRECATION_WARNINGS
+			CollisionData.QueryData = BuildLegacyQueryFilter(OwnerId, QueryBlockMask, QueryOverlapMask, QueryMaskFilter, QueryCollisionChannel, QueryFlags);
+			CollisionData.SimData = BuildLegacySimFilter(BodyIndex, SimBlockMask, ComponentId, SimMaskFilter, SimCollisionChannel, SimFlags);
+			PRAGMA_ENABLE_DEPRECATION_WARNINGS
+
+			Chaos::Tests::Save(CollisionData, Version, SerializationData);
+		}
+
+		FCollisionData Load(const int32 Version) const
+		{
+			FCollisionData ResultCollisionData;
+			Chaos::Tests::Load(SerializationData, Version, ResultCollisionData);
+			return ResultCollisionData;
+		}
+	};
+
+	TEST_CASE("FCollisionData: Collision Filters Round Trip", "[Chaos][CollisionData][unit]")
+	{
+		const Chaos::Filter::FInstanceData InstanceData(1, 2);
+		Chaos::Filter::FShapeFilterBuilder Builder;
+		Builder.SetBlockChannelMask(0xFEDCBA976543210);
+		Builder.SetOverlapChannelMask(0x0123456789ABCDEF);
+		Builder.SetMaskFilter(33); // 0b100001
+		Builder.SetCollisionChannelIndex(63);
+		Builder.SetFilterFlags(EFilterFlags::SimpleCollision | EFilterFlags::ComplexCollision | EFilterFlags::ContactNotify);
+		const Chaos::Filter::FShapeFilterData ShapeFilter = Builder.Build();
+
+		FCollisionData CollisionData;
+		CollisionData.SetFilterInstanceData(InstanceData);
+		CollisionData.SetShapeFilterData(ShapeFilter);
+		const Chaos::Filter::FInstanceData ResultInstanceData = CollisionData.GetFilterInstanceData();
+		const Chaos::Filter::FShapeFilterData ResultShapeFilter = CollisionData.GetShapeFilterData();
+
+		CHECK(ResultInstanceData.GetOwnerId() == InstanceData.GetOwnerId());
+		CHECK(ResultInstanceData.GetComponentId() == InstanceData.GetComponentId());
+		CHECK(ResultShapeFilter.GetOverlapChannels() == ShapeFilter.GetOverlapChannels());
+		CHECK(ResultShapeFilter.GetBlockChannels() == ShapeFilter.GetBlockChannels());
+		CHECK(ResultShapeFilter.GetMaskFilter() == ShapeFilter.GetMaskFilter());
+		CHECK(ResultShapeFilter.GetCollisionChannelIndex() == ShapeFilter.GetCollisionChannelIndex());
+		CHECK(ResultShapeFilter.GetCollisionChannelMask() == ShapeFilter.GetCollisionChannelMask());
+		CHECK(ResultShapeFilter.GetFlags() == ShapeFilter.GetFlags());
+	}
+
+	TEST_CASE("FCollisionData: CollisionFilter Serialization: Pre 64Bit -> 64Bit: Combined", "[Chaos][CollisionData][unit]")
+	{
+		constexpr int32 Version = FFortniteMainBranchObjectVersion::CollisionFilter64Bit - 1;
+
+		TArray<FCombinedTestData> Tests
+		{
+			FCombinedTestData
+			{
+				.OwnerId = 0x12345678,
+				.SerializationData = TArray<uint8>{0x78, 0x56, 0x34, 0x12, 0, 0, 0, 0, 0, 0, 0, 0, 0, 0, 0, 0, 0, 0, 0, 0, 0, 0, 0, 0, 0, 0, 0, 0, 0, 0, 0, 0, 1, 1, 0, 0, 0, 0, 0, 0, 0, 0, 0, },
+			},
+			FCombinedTestData
+			{
+				.OverlapMask = 0x12345678,
+				.SerializationData = TArray<uint8>{0, 0, 0, 0, 0, 0, 0, 0, 0x78, 0x56, 0x34, 0x12, 0, 0, 0, 0, 0, 0, 0, 0, 0, 0, 0, 0, 0, 0, 0, 0, 0, 0, 0, 0, 1, 1, 0, 0, 0, 0, 0, 0, 0, 0, 0, },
+			},
+			FCombinedTestData
+			{
+				.BlockMask = 0x12345678,
+				.SerializationData = TArray<uint8>{0, 0, 0, 0, 0x78, 0x56, 0x34, 0x12, 0, 0, 0, 0, 0, 0, 0, 0, 0, 0, 0, 0, 0x78, 0x56, 0x34, 0x12, 0, 0, 0, 0, 0, 0, 0, 0, 1, 1, 0, 0, 0, 0, 0, 0, 0, 0, 0, },
+			},
+			FCombinedTestData
+			{
+				.ComponentId = 0x12345678,
+				.SerializationData = TArray<uint8>{0, 0, 0, 0, 0, 0, 0, 0, 0, 0, 0, 0, 0, 0, 0, 0, 0, 0, 0, 0, 0, 0, 0, 0, 0x78, 0x56, 0x34, 0x12, 0, 0, 0, 0, 1, 1, 0, 0, 0, 0, 0, 0, 0, 0, 0, },
+			},
+			FCombinedTestData
+			{
+				.Flags = EFilterFlags::All, // Word3: 000000 00000 000000000000011111111 -> 0x0000 00FF
+				.SerializationData = TArray<uint8>{0, 0, 0, 0, 0, 0, 0, 0, 0, 0, 0, 0, 0xFF, 0, 0, 0, 0, 0, 0, 0, 0, 0, 0, 0, 0, 0, 0, 0, 0xFF, 0, 0, 0, 1, 1, 0, 0, 0, 0, 0, 0, 0, 0, 0, },
+			},
+			FCombinedTestData
+			{
+				.MaskFilter = 0x21, // Word3: 100001 00000 000000000000000000000 -> 0x8400 0000
+				.SerializationData = TArray<uint8>{0, 0, 0, 0, 0, 0, 0, 0, 0, 0, 0, 0, 0, 0, 0, 0x84, 0, 0, 0, 0, 0, 0, 0, 0, 0, 0, 0, 0, 0, 0, 0, 0x84, 1, 1, 0, 0, 0, 0, 0, 0, 0, 0, 0, },
+			},
+			FCombinedTestData
+			{
+				.CollisionChannel = 31, // Word3: 000000 11111 000000000000000000000 -> 0x03E0 0000
+				.SerializationData = TArray<uint8>{0, 0, 0, 0, 0, 0, 0, 0, 0, 0, 0, 0, 0, 0, 0xE0, 0x03, 0, 0, 0, 0, 0, 0, 0, 0, 0, 0, 0, 0, 0, 0, 0xE0, 0x03, 1, 1, 0, 0, 0, 0, 0, 0, 0, 0, 0, },
+			},
+			FCombinedTestData
+			{
+				.ComponentId = 1,
+				.OwnerId = 2,
+				.OverlapMask = 0xABCD4567,
+				.BlockMask = 0x12345678,
+				.MaskFilter = 0x21,
+				.CollisionChannel = 31,
+				.Flags = EFilterFlags::SimpleCollision | EFilterFlags::ComplexCollision | EFilterFlags::ContactNotify,
+				.SerializationData = TArray<uint8>{2, 0, 0, 0, 0x78, 0x56, 0x34, 0x12, 0x67, 0x45, 0xCD, 0xAB, 0x0B, 0, 0xE0, 0x87, 0, 0, 0, 0, 0x78, 0x56, 0x34, 0x12, 1, 0, 0, 0, 0x0B, 0, 0xE0, 0x87, 1, 1, 0, 0, 0, 0, 0, 0, 0, 0, 0, },
+			},
+		};
+		for (int32 I = 0; I < Tests.Num(); ++I)
+		{
+			CAPTURE(I);
+			const FCombinedTestData& TestData = Tests[I];
+			const FCollisionData ResultCollisionData = TestData.Load(Version);
+
+			const Chaos::Filter::FInstanceData& InstanceData = ResultCollisionData.GetFilterInstanceData();
+			const Chaos::Filter::FShapeFilterData& ShapeFilter = ResultCollisionData.GetShapeFilterData();
+
+			CHECK(TestData.OwnerId == InstanceData.GetOwnerId());
+			CHECK(TestData.ComponentId == InstanceData.GetComponentId());
+			CHECK(TestData.OverlapMask == ShapeFilter.GetOverlapChannels());
+			CHECK(TestData.BlockMask == ShapeFilter.GetBlockChannels());
+			CHECK(TestData.MaskFilter == ShapeFilter.GetMaskFilter());
+			CHECK(TestData.CollisionChannel == ShapeFilter.GetCollisionChannelIndex());
+			CHECK(TestData.Flags == ShapeFilter.GetFlags());
+		}
+	}
+
+	TEST_CASE("FCollisionData: CollisionFilter Serialization: Pre 64Bit -> 64Bit: Complex", "[Chaos][CollisionData][unit]")
+	{
+		constexpr int32 Version = FFortniteMainBranchObjectVersion::CollisionFilter64Bit - 1;
+
+		TArray<FLegacyFilterTestData> Tests
+		{
+			FLegacyFilterTestData
+			{
+				.QueryBlockMask = 0x12345678,
+				.SimBlockMask = 0xAABCDEFF,
+				.ExpectedBlockMask = 0x12345678,
+				.SerializationData = TArray<uint8>{0, 0, 0, 0, 0x78, 0x56, 0x34, 0x12, 0, 0, 0, 0, 0, 0, 0, 0, 0, 0, 0, 0, 0xFF, 0xDE, 0xBC, 0xAA, 0, 0, 0, 0, 0, 0, 0, 0, 1, 1, 0, 0, 0, 0, 0, },
+			},
+			FLegacyFilterTestData
+			{
+				.QueryCollisionChannel = 1,
+				.SimCollisionChannel = 2,
+				.ExpectedCollisionChannel = 1,
+				.SerializationData = TArray<uint8>{0, 0, 0, 0, 0, 0, 0, 0, 0, 0, 0, 0, 0, 0, 0x20, 0, 0, 0, 0, 0, 0, 0, 0, 0, 0, 0, 0, 0, 0, 0, 0x40, 0, 0x1, 0x1, 0, 0, 0, 0, 0, },
+			},
+			FLegacyFilterTestData
+			{
+				.QueryFlags = EFilterFlags::SimpleCollision | EFilterFlags::ComplexCollision, // 0x3
+				.SimFlags = EFilterFlags::ContactNotify, // 0xA
+				.ExpectedFlags = EFilterFlags::SimpleCollision | EFilterFlags::ComplexCollision | EFilterFlags::ContactNotify,
+				.SerializationData = TArray<uint8>{0, 0, 0, 0, 0, 0, 0, 0, 0, 0, 0, 0, 0x3, 0, 0, 0, 0, 0, 0, 0, 0, 0, 0, 0, 0, 0, 0, 0, 0x8, 0, 0, 0, 0x1, 0x1, 0, 0, 0, 0, 0, },
+			},
+			FLegacyFilterTestData
+			{
+				.QueryFlags = EFilterFlags::ContactNotify, // 0xA
+				.SimFlags = EFilterFlags::SimpleCollision | EFilterFlags::ComplexCollision, // 0x3
+				.ExpectedFlags = EFilterFlags::None,
+				.SerializationData = TArray<uint8>{0, 0, 0, 0, 0, 0, 0, 0, 0, 0, 0, 0, 0x8, 0, 0, 0, 0, 0, 0, 0, 0, 0, 0, 0, 0, 0, 0, 0, 0x3, 0, 0, 0, 0x1, 0x1, 0, 0, 0, 0, 0, },
+			},
+			FLegacyFilterTestData
+			{
+				.QueryMaskFilter = 0xA, // Word3: 001010 00000 000000000000000000000 -> 0x2800 0000
+				.SimMaskFilter = 0xF, // Word3: 001111 00000 000000000000000000000 -> 0x3C00 0000
+				.ExpectedMaskFilter = 0xA,
+				.SerializationData = TArray<uint8>{0, 0, 0, 0, 0, 0, 0, 0, 0, 0, 0, 0, 0, 0, 0, 0x28, 0, 0, 0, 0, 0, 0, 0, 0, 0, 0, 0, 0, 0, 0, 0, 0x3C, 1, 1, 0, 0, 0, 0, 0, },
+			},
+		};
+		for (int32 I = 0; I < Tests.Num(); ++I)
+		{
+			CAPTURE(I);
+			const FLegacyFilterTestData& TestData = Tests[I];
+			const FCollisionData ResultCollisionData = TestData.Load(Version);
+
+			const Chaos::Filter::FInstanceData& InstanceData = ResultCollisionData.GetFilterInstanceData();
+			const Chaos::Filter::FShapeFilterData& ShapeFilter = ResultCollisionData.GetShapeFilterData();
+
+			CHECK(TestData.OwnerId == InstanceData.GetOwnerId());
+			CHECK(TestData.ComponentId == InstanceData.GetComponentId());
+			CHECK(TestData.QueryOverlapMask == ShapeFilter.GetOverlapChannels());
+			CHECK(TestData.ExpectedBlockMask == ShapeFilter.GetBlockChannels());
+			CHECK(TestData.ExpectedMaskFilter == ShapeFilter.GetMaskFilter());
+			CHECK(TestData.ExpectedCollisionChannel == ShapeFilter.GetCollisionChannelIndex());
+			CHECK(TestData.ExpectedFlags == ShapeFilter.GetFlags());
+		}
+	}
+
+	TEST_CASE("FCollisionData: CollisionFilter Serialization: Round Trip", "[Chaos][CollisionData][unit]")
+	{
+		constexpr int32 Version = FFortniteMainBranchObjectVersion::CollisionFilter64Bit;
+
+		TArray<FCombinedTestData> Tests
+		{
+			FCombinedTestData
+			{
+				.OwnerId = 0x12345678,
+			},
+			FCombinedTestData
+			{
+				.OverlapMask = 0x12345678,
+			},
+			FCombinedTestData
+			{
+				.OverlapMask = 0xFEDCBA976543210,
+			},
+			FCombinedTestData
+			{
+				.BlockMask = 0x12345678,
+			},
+			FCombinedTestData
+			{
+				.BlockMask = 0xFEDCBA976543210,
+			},
+			FCombinedTestData
+			{
+				.ComponentId = 0x12345678,
+			},
+			FCombinedTestData
+			{
+				.Flags = EFilterFlags::All,
+			},
+			FCombinedTestData
+			{
+				.MaskFilter = 0x21,
+			},
+			FCombinedTestData
+			{
+				.CollisionChannel = 31,
+			},
+			FCombinedTestData
+			{
+				.CollisionChannel = 63,
+			},
+			FCombinedTestData
+			{
+				.ComponentId = 1,
+				.OwnerId = 2,
+				.OverlapMask = 0xFEDCBA9876543210,
+				.BlockMask = 0x123456789ABCDEF,
+				.MaskFilter = 31,
+				.CollisionChannel = 63,
+				.Flags = EFilterFlags::SimpleCollision | EFilterFlags::ComplexCollision | EFilterFlags::ContactNotify,
+			},
+		};
+		for (int32 I = 0; I < Tests.Num(); ++I)
+		{
+			CAPTURE(I);
+			FCombinedTestData& TestData = Tests[I];
+
+			TestData.Save(Version);
+			const FCollisionData ResultCollisionData = TestData.Load(Version);
+
+			const Chaos::Filter::FInstanceData& InstanceData = ResultCollisionData.GetFilterInstanceData();
+			const Chaos::Filter::FShapeFilterData& ShapeFilter = ResultCollisionData.GetShapeFilterData();
+
+			CHECK(TestData.OwnerId == InstanceData.GetOwnerId());
+			CHECK(TestData.ComponentId == InstanceData.GetComponentId());
+			CHECK(TestData.OverlapMask == ShapeFilter.GetOverlapChannels());
+			CHECK(TestData.BlockMask == ShapeFilter.GetBlockChannels());
+			CHECK(TestData.MaskFilter == ShapeFilter.GetMaskFilter());
+			CHECK(TestData.CollisionChannel == ShapeFilter.GetCollisionChannelIndex());
+			CHECK(TestData.Flags == ShapeFilter.GetFlags());
+		}
+	}
+
+	TEST_CASE("CollisionData: Storage: Perf", "[Chaos][CollisionData][!benchmark]")
+	{
+		constexpr uint64 Iterations = 100000;
+		TArray<Chaos::Filter::FShapeFilterData> ShapeFilters;
+		TArray<Chaos::Filter::FInstanceData> InstanceDatas;
+		TArray<Chaos::Filter::FCombinedShapeFilterData> CombinedDatas;
+		ShapeFilters.SetNum(Iterations);
+		InstanceDatas.SetNum(Iterations);
+		CombinedDatas.SetNum(Iterations);
+		FCollisionData CollisionData;
+
+		BENCHMARK("Get ShapeFilterData")
+		{
+			for (uint64 I = 0; I < Iterations; ++I)
+			{
+				ShapeFilters[I] = CollisionData.GetShapeFilterData();
+			}
+		};
+		BENCHMARK("Get InstanceData")
+		{
+			for (uint64 I = 0; I < Iterations; ++I)
+			{
+				InstanceDatas[I] = CollisionData.GetFilterInstanceData();
+			}
+		};
+		BENCHMARK("Get CombinedFilter")
+		{
+			for (uint64 I = 0; I < Iterations; ++I)
+			{
+				CombinedDatas[I] = CollisionData.GetCombinedShapeFilterData();
+			}
+		};
+		BENCHMARK("Set ShapeFilterData")
+		{
+			for (uint64 I = 0; I < Iterations; ++I)
+			{
+				CollisionData.SetShapeFilterData(ShapeFilters[I]);
+			}
+		};
+		BENCHMARK("Set InstanceData")
+		{
+			for (uint64 I = 0; I < Iterations; ++I)
+			{
+				CollisionData.SetFilterInstanceData(InstanceDatas[I]);
+			}
+		};
+		BENCHMARK("Set CombinedFilter")
+		{
+			for (uint64 I = 0; I < Iterations; ++I)
+			{
+				CollisionData.SetCombinedShapeFilterData(CombinedDatas[I]);
+			}
+		};
+		BENCHMARK("GetSet ShapeFilterData")
+		{
+			for (uint64 I = 0; I < Iterations; ++I)
+			{
+				const Chaos::Filter::FShapeFilterData ShapeData = ShapeFilters[I];
+				ShapeFilters[I] = CollisionData.GetShapeFilterData();
+				CollisionData.SetShapeFilterData(ShapeData);
+			}
+		};
+		BENCHMARK("GetSet InstanceData")
+		{
+			for (uint64 I = 0; I < Iterations; ++I)
+			{
+				const Chaos::Filter::FInstanceData InstanceData = InstanceDatas[I];
+				InstanceDatas[I] = CollisionData.GetFilterInstanceData();
+				CollisionData.SetFilterInstanceData(InstanceData);
+			}
+		};
+		BENCHMARK("GetSet CombinedFilter")
+		{
+			for (uint64 I = 0; I < Iterations; ++I)
+			{
+				const Chaos::Filter::FCombinedShapeFilterData CombinedData = CombinedDatas[I];
+				CombinedDatas[I] = CollisionData.GetCombinedShapeFilterData();
+				CollisionData.SetCombinedShapeFilterData(CombinedData);
+			}
+		};
+		BENCHMARK("GetSet BothFilters")
+		{
+			for (uint64 I = 0; I < Iterations; ++I)
+			{
+				const Chaos::Filter::FShapeFilterData ShapeData = ShapeFilters[I];
+				const Chaos::Filter::FInstanceData InstanceData = InstanceDatas[I];
+
+				ShapeFilters[I] = CollisionData.GetShapeFilterData();
+				InstanceDatas[I] = CollisionData.GetFilterInstanceData();
+
+				CollisionData.SetShapeFilterData(ShapeData);
+				CollisionData.SetFilterInstanceData(InstanceData);
+			}
+		};
+	}
+} // namespace Chaos::Tests
diff --git a/Engine/Source/Programs/ChaosPhysicsTests/Private/CollisionEventTests.cpp b/Engine/Source/Programs/ChaosPhysicsTests/Private/CollisionEventTests.cpp
new file mode 100644
--- /dev/null
+++ b/Engine/Source/Programs/ChaosPhysicsTests/Private/CollisionEventTests.cpp
@@ -0,0 +1,157 @@
+// Copyright Epic Games, Inc. All Rights Reserved.
+
+#include "CoreMinimal.h"
+#include "ChaosTestHarness.h"
+#include "ChaosTestScene.h"
+
+#include "Chaos/Box.h"
+#include "Chaos/MassProperties.h"
+#include "ChaosSolversModule.h"
+#include "EventsData.h"
+#include "PBDRigidsSolver.h"
+#include "PhysicsProxy/SingleParticlePhysicsProxy.h"
+
+TEST_CASE("CollisionEventTests", "[Chaos][Collision][Events][unit]")
+{
+	using namespace Chaos;
+	using namespace Chaos::LowLevelTest;
+
+	const float Dt = 1.0f / 30.0f;
+
+	FChaosTestScene Scene;
+
+	//
+	// Collision events are dispatched every frame and contain collision data for awake bodies.
+	//
+	SECTION("Collision Events dispatched for awake bodies")
+	{
+		const double KinematicPositionTolerance = UE_SMALL_NUMBER;
+		const double DynamicPositionTolerance = UE_KINDA_SMALL_NUMBER;
+
+		// Place a dynamic 100cm box exactly on the floor (a large kinematic box)
+		FSingleParticlePhysicsProxy* FloorBox = Scene.AddKinematicBox(TEXT("Floor"), FTransform3d(FVector3d(0, 0, -50)), FVector3d(500, 500, 50));
+		FSingleParticlePhysicsProxy* DynamicBox = Scene.AddDynamicBox(TEXT("Box"), FTransform3d(FVector3d(0, 0, 50)), FVector3d(50, 50, 50));
+
+		int NumCallbacks = 0;
+		Scene.SetCollisionCallback(
+			[&NumCallbacks, DynamicBox](const FCollisionEventData& CollisionEventData)
+			{
+				// Should have zero or one collision event
+				CHECK(CollisionEventData.CollisionData.AllCollisionsArray.Num() >= 0);
+				CHECK(CollisionEventData.CollisionData.AllCollisionsArray.Num() <= 1);
+				++NumCallbacks;
+			});
+
+		// The box should sleep quickly (but we should have at least one tick)
+		Scene.TickSolverToSleep(Dt, 50);
+		REQUIRE(Scene.IsSceneSleeping());
+		CHECK(Scene.GetNumTicks() > 0);
+
+		// We should have one callback for each tick
+		CHECK(NumCallbacks == Scene.GetNumTicks());
+
+		// The kinematic box should not have moved at all
+		FVector3d FloorBoxPos = FloorBox->GetGameThreadAPI().GetX();
+		CHECK_THAT(FloorBoxPos.Z, Catch::Matchers::WithinAbs(-50.0f, KinematicPositionTolerance));
+
+		// The dynamic box should not have moved much
+		FVector3d DynamicBoxPos = DynamicBox->GetGameThreadAPI().GetX();
+		CHECK_THAT(DynamicBoxPos.Z, Catch::Matchers::WithinAbs(50.0f, DynamicPositionTolerance));
+	}
+
+
+	//
+	// The collision event callback should still be called even when there are no collision to report. This
+	// is so that user code can track when contact ends.
+	//
+	SECTION("Collision Events dispatched empty when no awake bodies")
+	{
+		FSingleParticlePhysicsProxy* FloorBox = Scene.AddKinematicBox(TEXT("Floor"), FTransform3d(FVector3d(0, 0, -50)), FVector3d(500, 500, 50));
+		FSingleParticlePhysicsProxy* DynamicBox1 = Scene.AddDynamicBox(TEXT("Box1"), FTransform3d(FVector3d(0, 0, 50)), FVector3d(50, 50, 50));
+
+		// Wait for scene to sleep
+		Scene.TickSolverToSleep(Dt, 1000);
+		REQUIRE(Scene.IsSceneSleeping());
+
+		// Make sure the box actually went to sleep or the test isn't working...
+		REQUIRE(DynamicBox1->GetGameThreadAPI().ObjectState() == EObjectStateType::Sleeping);
+
+		// Attach the collision callback
+		int32 NumCallbacks = 0;
+		int32 NumCollisions = 0;
+		Scene.SetCollisionCallback(
+			[&NumCallbacks, &NumCollisions](const FCollisionEventData& CollisionEventData)
+			{
+				NumCollisions += CollisionEventData.CollisionData.AllCollisionsArray.Num();
+				NumCallbacks += 1;
+			});
+
+		// Tick the sleeping scene
+		Scene.TickSolver(Dt);
+
+		// Callback should have been fired, even though we had no collisions
+		CHECK(NumCallbacks == 1);
+		CHECK(NumCollisions == 0);
+
+		// If we wake the body we should get collision events again
+		DynamicBox1->GetGameThreadAPI().SetObjectState(EObjectStateType::Dynamic);
+
+		NumCallbacks = 0;
+		NumCollisions = 0;
+		Scene.TickSolver(Dt);
+
+		// Callback should have been called again and we should have 1 collison (Box1-Floor)
+		CHECK(NumCallbacks == 1);
+		CHECK(NumCollisions == 1);
+	}
+
+
+	//
+	// BUG: Collision events were not dispatched on the first frame after waking.
+	// Collisions were not being re-added to the ActiveConstraints list when
+	// awoken naturally from island modification (from IslandManager). 
+	//
+	SECTION("Collision Events dispatched empty when no awake bodies")
+	{
+		FSingleParticlePhysicsProxy* FloorBox = Scene.AddKinematicBox(TEXT("Floor"), FTransform3d(FVector3d(0, 0, -50)), FVector3d(500, 500, 50));
+		FSingleParticlePhysicsProxy* DynamicBox1 = Scene.AddDynamicBox(TEXT("Box1"), FTransform3d(FVector3d(0, 0, 50)), FVector3d(50, 50, 50));
+
+		// Wait for sleep
+		bool IsSceneSleeping = Scene.TickSolverToSleep(Dt, 1000);
+		REQUIRE(IsSceneSleeping);
+		REQUIRE(DynamicBox1->GetGameThreadAPI().ObjectState() == EObjectStateType::Sleeping);
+
+		// Attach collision callback
+		int32 NumCallbacks = 0;
+		int32 NumCollisions = 0;
+		Scene.SetCollisionCallback(
+			[&NumCallbacks, &NumCollisions](const FCollisionEventData& CollisionEventData)
+			{
+				NumCollisions += CollisionEventData.CollisionData.AllCollisionsArray.Num();
+				NumCallbacks += 1;
+			});
+
+		Scene.TickSolver(Dt);
+		CHECK(NumCallbacks == 1);
+		CHECK(NumCollisions == 0);
+
+		// Spawn a new awake box which is colliding with the sleeping box. This should
+		// wake it naturally (as opposed to explicitly like above) which is a separate
+		// code path to be tested (in IslandManager.cpp). For this to work, the
+		// IslandManager must be adding awakened collisions back into ActiveConstraints.
+		FSingleParticlePhysicsProxy* DynamicBox2 = Scene.AddDynamicBox(TEXT("Box2"), FTransform3d(FVector3d(0, 0, 95)), FVector3d(50, 50, 50));
+
+		NumCallbacks = 0;
+		NumCollisions = 0;
+		Scene.TickSolver(Dt);
+
+		// Callback was called and we have two collisions
+		// (Prior to bug fix we only had 1 collision)
+		CHECK(NumCallbacks == 1);
+		CHECK(NumCollisions == 2);
+
+		// Both boxes should be awake
+		CHECK(DynamicBox1->GetGameThreadAPI().ObjectState() != EObjectStateType::Sleeping);
+		CHECK(DynamicBox2->GetGameThreadAPI().ObjectState() != EObjectStateType::Sleeping);
+	}
+}
diff --git a/Engine/Source/Programs/ChaosPhysicsTests/Private/FilterTests.cpp b/Engine/Source/Programs/ChaosPhysicsTests/Private/FilterTests.cpp
new file mode 100644
--- /dev/null
+++ b/Engine/Source/Programs/ChaosPhysicsTests/Private/FilterTests.cpp
@@ -0,0 +1,1173 @@
+// Copyright Epic Games, Inc. All Rights Reserved.
+#include "CoreMinimal.h"
+
+#include "ChaosTestHarness.h"
+#include <catch2/generators/catch_generators_range.hpp>
+
+#include "Chaos/CollisionFilterData.h"
+
+namespace Chaos::Filter
+{
+
+constexpr Chaos::EFilterFlags EmptyFilterFlags = Chaos::EFilterFlags::None;
+constexpr uint64 FullChannelMask = std::numeric_limits<uint64>::max();
+constexpr uint64 MaxChannelBitCount = 64;
+
+uint64 ToChannelMask(uint8 ChannelIndex)
+{
+	return 1LLU << ChannelIndex;
+}
+
+FShapeFilterData BuildShapeFilter(const uint8 ChannelIndex, const uint64 BlockMask, const uint64 OverlapMask)
+{
+	return FShapeFilterBuilder()
+		.SetCollisionChannelIndex(ChannelIndex)
+		.SetBlockChannelMask(BlockMask)
+		.SetOverlapChannelMask(OverlapMask)
+		.Build();
+}
+
+FShapeUnionFilterData BuildUnionFilter(const TArray<FShapeFilterData>& ShapeFilters)
+{
+	FShapeUnionFilterData UnionShape;
+	for (const FShapeFilterData& ShapeFilter : ShapeFilters)
+	{
+		UnionShape.Combine(ShapeFilter);
+	}
+	return UnionShape;
+}
+
+
+TEST_CASE("FInstanceData - Construction", "[Chaos][Filter]")
+{
+	const uint32 OwnerId = 1;
+	const uint32 ComponentId = 2;
+	FInstanceData InstanceData(OwnerId, ComponentId);
+
+	CHECK(OwnerId == InstanceData.GetOwnerId());
+	CHECK(ComponentId == InstanceData.GetComponentId());
+}
+
+TEST_CASE("FInstanceData - IsValid", "[Chaos][Filter]")
+{
+	FInstanceData InvalidInstanceData;
+
+	CHECK(FInstanceData().IsValid() == false);
+	CHECK(FInstanceData(1, 0).IsValid() == true);
+	CHECK(FInstanceData(0, 1).IsValid() == true);
+	CHECK(FInstanceData(1, 1).IsValid() == true);
+}
+
+TEST_CASE("FInstanceData - ToString", "[Chaos][Filter]")
+{
+	const uint32 OwnerId = 1;
+	const uint32 ComponentId = 2;
+	FInstanceData InstanceData(OwnerId, ComponentId);
+
+	CHECK(InstanceData.ToString() == FString("OwnerId(1) ComponentId(2)"));
+}
+
+TEST_CASE("FShapeFilterData - DefaultConstructor", "[Chaos][Filter]")
+{
+	FShapeFilterData ShapeFilter;
+	CHECK(!ShapeFilter.IsValid());
+	CHECK(ShapeFilter.GetCollisionChannelIndex() == 0);
+	CHECK(ShapeFilter.GetOverlapChannels() == 0);
+	CHECK(ShapeFilter.GetBlockChannels() == 0);
+	CHECK(ShapeFilter.GetMaskFilter() == 0);
+	CHECK(ShapeFilter.GetFlags() == EmptyFilterFlags);
+}
+
+TEST_CASE("FShapeFilterData - MaskFilter", "[Chaos][Filter]")
+{
+	constexpr int8 MaskFilterBitCount = 6;
+	constexpr int8 MaskFilter = static_cast<int8>((1 << MaskFilterBitCount) - 1);
+
+	FShapeFilterData ShapeFilter;
+	ShapeFilter.SetMaskFilter(MaskFilter);
+
+	const int8 ActualMaskFilter = ShapeFilter.GetMaskFilter();
+	CHECK(ActualMaskFilter == MaskFilter);
+}
+
+TEST_CASE("FShapeFilterData - IsCollisionChannelSet", "[Chaos][Filter]")
+{
+	const uint32 ChannelIndex = GENERATE(0, 1, 16, 31, 32, 63);
+	CAPTURE(ChannelIndex);
+
+	// TODO @ JoshD: Update for channels as mask.
+	FShapeFilterBuilder Builder;
+	Builder.SetCollisionChannelIndex(ChannelIndex);
+	const FShapeFilterData FilterData = Builder.Build();
+
+	for (uint32 TestIndex = 0; TestIndex < MaxChannelBitCount; ++TestIndex)
+	{
+		CAPTURE(TestIndex);
+		CHECK((TestIndex == ChannelIndex) == FilterData.IsCollisionChannelSet(TestIndex));
+	}
+}
+
+TEST_CASE("FShapeFilterData - ToString", "[Chaos][Filter]")
+{
+	FShapeFilterBuilder Builder;
+	Builder.SetCollisionChannelIndex(1);
+	Builder.SetBlockChannelMask(0xAAAAAAAA);
+	Builder.SetOverlapChannelMask(0xCCCCCCCC);
+	Builder.SetMaskFilter(0x3F);
+	Builder.SetFilterFlags(Chaos::EFilterFlags::SimpleCollision | Chaos::EFilterFlags::ComplexCollision);
+	const FShapeFilterData FilterData = Builder.Build();
+
+	CHECK(FilterData.ToString() == "Block(0xAAAAAAAA) Overlap(0xCCCCCCCC) CollisionChannel(0x1) MaskFilter(0x3F) Flags(0x3)");
+}
+
+TEST_CASE("FQueryFilterData - DefaultConstructor", "[Chaos][Filter]")
+{
+	const FQueryFilterData QueryFilterData;
+	CHECK(!QueryFilterData.IsValid());
+	CHECK(QueryFilterData.GetQueryType() == FQueryFilterData::EQueryType::ObjectType);
+	CHECK(QueryFilterData.GetObjectTypesToQueryMask() == 0);
+	CHECK(QueryFilterData.IsMultiQuery() == false);
+	CHECK(QueryFilterData.GetIgnoreMask() == 0);
+	CHECK(QueryFilterData.GetFlags() == EmptyFilterFlags);
+}
+
+TEST_CASE("FQueryFilterData - IsCollisionChannelSet", "[Chaos][Filter]")
+{
+	// TODO @ JoshD: Update for channels as mask.
+	const uint32 ChannelIndex = GENERATE(0, 1, 16, 31, 32, 63);
+	CAPTURE(ChannelIndex);
+
+	FQueryTraceFilterBuilder Builder;
+	Builder.SetCollisionChannelIndex(ChannelIndex);
+	const FQueryFilterData FilterData = Builder.Build();
+
+	for (uint32 TestIndex = 0; TestIndex < MaxChannelBitCount; ++TestIndex)
+	{
+		CAPTURE(TestIndex);
+		CHECK((TestIndex == ChannelIndex) == FilterData.IsCollisionChannelSet(TestIndex));
+	}
+}
+
+TEST_CASE("FQueryFilterData::ObjectFilter - ToString", "[Chaos][Filter]")
+{
+	FQueryObjectFilterBuilder Builder;
+	Builder.SetObjectTypes(0xFF00FF);
+	Builder.SetMultiQuery(true);
+	Builder.SetMaskFilter(0x3F);
+	Builder.SetFilterFlags(Chaos::EFilterFlags::SimpleCollision);
+	const FQueryFilterData FilterData = Builder.Build();
+
+	CHECK(FilterData.ToString() == "Type(Object) IsMulti(1) ObjectTypes(0xFF00FF) IgnoreMask(0x3F) Flags(0x1)");
+}
+
+TEST_CASE("FQueryFilterData::TraceFilter - ToString", "[Chaos][Filter]")
+{
+	FQueryTraceFilterBuilder Builder;
+	Builder.SetCollisionChannelIndex(1);
+	Builder.SetBlockChannelMask(0xAAAAAAAA);
+	Builder.SetOverlapChannelMask(0xCCCCCCCC);
+	Builder.SetMaskFilter(0x3F);
+	Builder.SetFilterFlags(Chaos::EFilterFlags::SimpleCollision | Chaos::EFilterFlags::ComplexCollision);
+	const FQueryFilterData FilterData = Builder.Build();
+
+	CHECK(FilterData.ToString() == "Type(Channel) Block(0xAAAAAAAA) Overlap(0xCCCCCCCC) ChannelIndex(0x1) IgnoreMask(0x3F) Flags(0x3)");
+}
+
+TEST_CASE("FCombinedShapeFilterData - DefaultConstructor", "[Chaos][Filter]")
+{
+	CHECK(!FCombinedShapeFilterData().IsValid());
+}
+
+TEST_CASE("FCombinedShapeFilterData - IsValid", "[Chaos][Filter]")
+{
+	const FInstanceData InstanceData(1, 2);
+	const FShapeFilterData ShapeFilterData = FShapeFilterBuilder().SetCollisionChannelIndex(1).Build();
+
+	CHECK(FCombinedShapeFilterData(ShapeFilterData, FInstanceData()).IsValid());
+	CHECK(FCombinedShapeFilterData(FShapeFilterData(), InstanceData).IsValid());
+}
+
+TEST_CASE("FCombinedShapeFilterData - InstanceData", "[Chaos][Filter]")
+{
+	const FInstanceData InstanceData(1, 2);
+	FCombinedShapeFilterData CombinedShapeFilterData;
+	CombinedShapeFilterData.SetInstanceData(InstanceData);
+
+	const FInstanceData ActualInstanceData = CombinedShapeFilterData.GetInstanceData();
+	CHECK(ActualInstanceData.GetOwnerId() == InstanceData.GetOwnerId());
+	CHECK(ActualInstanceData.GetComponentId() == InstanceData.GetComponentId());
+	CHECK(!CombinedShapeFilterData.GetShapeFilterData().IsValid());
+}
+
+TEST_CASE("FCombinedShapeFilterData - ShapeFilterData", "[Chaos][Filter]")
+{
+	FShapeFilterBuilder Builder;
+	Builder.SetCollisionChannelIndex(1);
+	Builder.SetOverlapChannelMask(2);
+	Builder.SetBlockChannelMask(3);
+	Builder.SetMaskFilter(4);
+	Builder.SetFilterFlags((Chaos::EFilterFlags)6);
+	const FShapeFilterData ShapeFilterData = Builder.Build();
+	FCombinedShapeFilterData CombinedShapeFilterData;
+	CombinedShapeFilterData.SetShapeFilterData(ShapeFilterData);
+
+	const FShapeFilterData ActualShapeFilterData = CombinedShapeFilterData.GetShapeFilterData();
+	CHECK(ActualShapeFilterData.GetCollisionChannelIndex() == ShapeFilterData.GetCollisionChannelIndex());
+	CHECK(ActualShapeFilterData.GetOverlapChannels() == ShapeFilterData.GetOverlapChannels());
+	CHECK(ActualShapeFilterData.GetBlockChannels() == ShapeFilterData.GetBlockChannels());
+	CHECK(ActualShapeFilterData.GetMaskFilter() == ShapeFilterData.GetMaskFilter());
+	CHECK(ActualShapeFilterData.GetFlags() == ShapeFilterData.GetFlags());
+	CHECK(!CombinedShapeFilterData.GetInstanceData().IsValid());
+}
+
+TEST_CASE("FShapeFilterBuilder - LegacyRoundTripConstruction", "[Chaos][Filter]")
+{
+	// For each section (flags, collision channel, mask filter) test first bit, last bit, all bits set. Throw in one last test for everything being set.
+	const uint32 Word3 = GENERATE(0, 0x1, 0x80, 0xFF, 1 << 21, 0x10 << 21, 0x1F << 21, 1 << 26, 0x20 << 26, 0x3F << 26, 0xFFE000FF);
+
+	const uint32 BlockChannels = 2;
+	FCollisionFilterData QueryData, SimData;
+	QueryData.Word0 = 1;
+	QueryData.Word1 = BlockChannels;
+	QueryData.Word2 = 3;
+	QueryData.Word3 = Word3;
+	SimData.Word0 = 0; // BodyIndex is no longer used
+	SimData.Word1 = BlockChannels;
+	SimData.Word2 = 7;
+	SimData.Word3 = Word3;
+
+	FCombinedShapeFilterData CombinedShapeFilter = FShapeFilterBuilder::BuildFromLegacyShapeFilter(QueryData, SimData);
+
+	FCollisionFilterData ResultQueryData, ResultSimData;
+	FShapeFilterBuilder::GetLegacyShapeFilter(CombinedShapeFilter, ResultQueryData, ResultSimData);
+	CHECK(ResultQueryData.Word0 == QueryData.Word0);
+	CHECK(ResultQueryData.Word1 == QueryData.Word1);
+	CHECK(ResultQueryData.Word2 == QueryData.Word2);
+	CHECK(ResultQueryData.Word3 == QueryData.Word3);
+	CHECK(ResultSimData.Word0 == SimData.Word0);
+	CHECK(ResultSimData.Word1 == SimData.Word1);
+	CHECK(ResultSimData.Word2 == SimData.Word2);
+	CHECK(ResultSimData.Word3 == SimData.Word3);
+}
+
+TEST_CASE("FShapeFilterBuilder - LegacyRoundTripConstruction - Separate Word3", "[Chaos][Filter]")
+{
+	const uint32 BlockChannels = 2;
+	FCollisionFilterData QueryData, SimData;
+	QueryData.Word0 = 1;
+	QueryData.Word1 = BlockChannels;
+	QueryData.Word2 = 3;
+	QueryData.Word3 = 0xFFE00003;
+	SimData.Word0 = 0; // BodyIndex is no longer used
+	SimData.Word1 = BlockChannels;
+	SimData.Word2 = 7;
+	SimData.Word3 = 0xFFE0000C;
+	const uint32 ExpectedWord3 = 0xFFE0000F;
+
+	FCombinedShapeFilterData CombinedShapeFilter = FShapeFilterBuilder::BuildFromLegacyShapeFilter(QueryData, SimData);
+
+	FCollisionFilterData ResultQueryData, ResultSimData;
+	FShapeFilterBuilder::GetLegacyShapeFilter(CombinedShapeFilter, ResultQueryData, ResultSimData);
+	CHECK(ResultQueryData.Word0 == QueryData.Word0);
+	CHECK(ResultQueryData.Word1 == QueryData.Word1);
+	CHECK(ResultQueryData.Word2 == QueryData.Word2);
+	CHECK(ResultQueryData.Word3 == ExpectedWord3);
+	CHECK(ResultSimData.Word0 == SimData.Word0);
+	CHECK(ResultSimData.Word1 == SimData.Word1);
+	CHECK(ResultSimData.Word2 == SimData.Word2);
+	CHECK(ResultSimData.Word3 == ExpectedWord3);
+}
+
+TEST_CASE("FShapeFilterBuilder - RoundTrip to legacy and back", "[Chaos][Filter]")
+{
+	const uint32 CollisionChannelIndex = GENERATE(0, 1, 31);
+	FShapeFilterBuilder Builder;
+	Builder.SetCollisionChannelIndex(CollisionChannelIndex);
+	Builder.SetBlockChannelMask(2);
+	Builder.SetOverlapChannelMask(4);
+	Builder.SetMaskFilter(5);
+	Builder.SetFilterFlags(EFilterFlags::All);
+	const FShapeFilterData ExpectedShapeFilter = Builder.Build();
+	const FInstanceData ExpectedInstanceData(9, 10);
+
+	FCollisionFilterData TempQueryData, TempSimData;
+
+	FShapeFilterBuilder::GetLegacyShapeFilter(FCombinedShapeFilterData(ExpectedShapeFilter, ExpectedInstanceData), TempQueryData, TempSimData);
+	const FCombinedShapeFilterData CombinedShapeFilter = FShapeFilterBuilder::BuildFromLegacyShapeFilter(TempQueryData, TempSimData);
+	const FShapeFilterData ResultShapeFilter = CombinedShapeFilter.GetShapeFilterData();
+	const FInstanceData ResultInstanceData = CombinedShapeFilter.GetInstanceData();
+
+	CHECK(ExpectedShapeFilter.GetCollisionChannelIndex() == ResultShapeFilter.GetCollisionChannelIndex());
+	CHECK(ExpectedShapeFilter.GetCollisionChannelMask() == ResultShapeFilter.GetCollisionChannelMask());
+	CHECK(ExpectedShapeFilter.GetBlockChannels() == ResultShapeFilter.GetBlockChannels());
+	CHECK(ExpectedShapeFilter.GetOverlapChannels() == ResultShapeFilter.GetOverlapChannels());
+	CHECK(ExpectedShapeFilter.GetMaskFilter() == ResultShapeFilter.GetMaskFilter());
+	CHECK(ExpectedShapeFilter.GetFlags() == ResultShapeFilter.GetFlags());
+	CHECK(ExpectedInstanceData.GetOwnerId() == ResultInstanceData.GetOwnerId());
+	CHECK(ExpectedInstanceData.GetComponentId() == ResultInstanceData.GetComponentId());
+}
+
+TEST_CASE("FShapeFilterBuilder - SetLegacyShapeQueryFilter", "[Chaos][Filter]")
+{
+	FCollisionFilterData QueryData0, QueryData1, SimData;
+	QueryData0.Word0 = 1;
+	QueryData0.Word1 = 2;
+	QueryData0.Word2 = 3;
+	QueryData0.Word3 = 0;
+	SimData.Word0 = 0; // BodyIndex is no longer used
+	SimData.Word1 = 2;
+	SimData.Word2 = 7;
+	SimData.Word3 = 0xC;
+	QueryData1.Word0 = 9;
+	QueryData1.Word1 = 10;
+	QueryData1.Word2 = 11;
+	QueryData1.Word3 = 0xFFE00003;
+	const uint32 ExpectedWord1 = 10;
+	const uint32 ExpectedWord3 = 0xFFE0000F;
+
+	FCombinedShapeFilterData CombinedShapeFilter = FShapeFilterBuilder::BuildFromLegacyShapeFilter(QueryData0, SimData);
+	FShapeFilterBuilder::SetLegacyShapeQueryFilter(CombinedShapeFilter, QueryData1);
+
+	const FCollisionFilterData ResultQueryData = FShapeFilterBuilder::GetLegacyShapeQueryFilter(CombinedShapeFilter);
+	const FCollisionFilterData ResultSimData = FShapeFilterBuilder::GetLegacyShapeSimFilter(CombinedShapeFilter);
+	CHECK(ResultQueryData.Word0 == QueryData1.Word0);
+	CHECK(ResultQueryData.Word1 == ExpectedWord1);
+	CHECK(ResultQueryData.Word2 == QueryData1.Word2);
+	CHECK(ResultQueryData.Word3 == ExpectedWord3);
+	CHECK(ResultSimData.Word0 == SimData.Word0);
+	CHECK(ResultSimData.Word1 == ExpectedWord1);
+	CHECK(ResultSimData.Word2 == SimData.Word2);
+	CHECK(ResultSimData.Word3 == ExpectedWord3);
+}
+
+TEST_CASE("FShapeFilterBuilder - SetLegacyShapeSimFilter", "[Chaos][Filter]")
+{
+	FCollisionFilterData QueryData, SimData0, SimData1;
+	QueryData.Word0 = 1;
+	QueryData.Word1 = 2;
+	QueryData.Word2 = 3;
+	QueryData.Word3 = 0x3;
+	SimData0.Word0 = 0; // BodyIndex is no longer used
+	SimData0.Word1 = 2;
+	SimData0.Word2 = 7;
+	SimData0.Word3 = 0;
+	SimData1.Word0 = 0; // BodyIndex is no longer used
+	SimData1.Word1 = 10;
+	SimData1.Word2 = 11;
+	SimData1.Word3 = 0xFFE0000C;
+	const uint32 ExpectedWord1 = 10;
+	const uint32 ExpectedWord3 = 0xFFE0000F;
+
+	FCombinedShapeFilterData CombinedShapeFilter = FShapeFilterBuilder::BuildFromLegacyShapeFilter(QueryData, SimData0);
+	FShapeFilterBuilder::SetLegacyShapeSimFilter(CombinedShapeFilter, SimData1);
+
+	const FCollisionFilterData ResultQueryData = FShapeFilterBuilder::GetLegacyShapeQueryFilter(CombinedShapeFilter);
+	const FCollisionFilterData ResultSimData = FShapeFilterBuilder::GetLegacyShapeSimFilter(CombinedShapeFilter);
+	CHECK(ResultQueryData.Word0 == QueryData.Word0);
+	CHECK(ResultQueryData.Word1 == ExpectedWord1);
+	CHECK(ResultQueryData.Word2 == QueryData.Word2);
+	CHECK(ResultQueryData.Word3 == ExpectedWord3);
+	CHECK(ResultSimData.Word0 == SimData1.Word0);
+	CHECK(ResultSimData.Word1 == ExpectedWord1);
+	CHECK(ResultSimData.Word2 == SimData1.Word2);
+	CHECK(ResultSimData.Word3 == ExpectedWord3);
+}
+
+TEST_CASE("FShapeFilterBuilder - SetCollisionChannelIndex", "[Chaos][Filter]")
+{
+	const uint8 ChannelIndex = GENERATE(0, 1, 3, 31, 32, 63);
+
+	FShapeFilterBuilder Builder;
+	Builder.SetCollisionChannelIndex(ChannelIndex);
+	const FShapeFilterData ShapeFilter = Builder.Build();
+
+	CHECK(ShapeFilter.GetCollisionChannelIndex() == ChannelIndex);
+	CHECK(ShapeFilter.GetCollisionChannelMask() == ToChannelMask(ChannelIndex));
+	CHECK(ShapeFilter.GetOverlapChannels() == 0);
+	CHECK(ShapeFilter.GetBlockChannels() == 0);
+	CHECK(ShapeFilter.GetMaskFilter() == 0);
+	CHECK(ShapeFilter.GetFlags() == EmptyFilterFlags);
+}
+
+TEST_CASE("FShapeFilterBuilder - SetBlockChannelMask", "[Chaos][Filter]")
+{
+	const TArray<uint64> BlockChannelMasks{ 0, 1, 2, 3, 0xAAAAAAAA, 0xFFFFFFFF, FullChannelMask };
+	const uint32 Index = GENERATE_COPY(Catch::Generators::range(0, BlockChannelMasks.Num()));
+	const uint64 BlockChannelMask = BlockChannelMasks[Index];
+
+	FShapeFilterBuilder Builder;
+	Builder.SetBlockChannelMask(BlockChannelMask);
+	const FShapeFilterData ShapeFilter = Builder.Build();
+
+	CHECK(ShapeFilter.GetCollisionChannelIndex() == 0);
+	CHECK(ShapeFilter.GetOverlapChannels() == 0);
+	CHECK(ShapeFilter.GetBlockChannels() == BlockChannelMask);
+	CHECK(ShapeFilter.GetMaskFilter() == 0);
+	CHECK(ShapeFilter.GetFlags() == EmptyFilterFlags);
+}
+
+TEST_CASE("FShapeFilterBuilder - SetOverlapChannelMask", "[Chaos][Filter]")
+{
+	const TArray<uint64> OverlapChannelMasks{ 0, 1, 2, 3, 0xAAAAAAAA, 0xFFFFFFFF, FullChannelMask };
+	const uint32 Index = GENERATE_COPY(Catch::Generators::range(0, OverlapChannelMasks.Num()));
+	const uint64 OverlapChannelMask = OverlapChannelMasks[Index];
+
+	FShapeFilterBuilder Builder;
+	Builder.SetOverlapChannelMask(OverlapChannelMask);
+	const FShapeFilterData ShapeFilter = Builder.Build();
+
+	CHECK(ShapeFilter.GetCollisionChannelIndex() == 0);
+	CHECK(ShapeFilter.GetOverlapChannels() == OverlapChannelMask);
+	CHECK(ShapeFilter.GetBlockChannels() == 0);
+	CHECK(ShapeFilter.GetMaskFilter() == 0);
+	CHECK(ShapeFilter.GetFlags() == EmptyFilterFlags);
+}
+
+TEST_CASE("FShapeFilterBuilder - SetMaskFilter", "[Chaos][Filter]")
+{
+	const TArray<uint8> MaskFilters{ 0, 1, 2, 0xF };
+	const uint32 Index = GENERATE(0, 1, 2, 3);
+	const uint8 MaskFilter = MaskFilters[Index];
+
+	FShapeFilterBuilder Builder;
+	Builder.SetMaskFilter(MaskFilter);
+	const FShapeFilterData ShapeFilter = Builder.Build();
+
+	CHECK(ShapeFilter.GetCollisionChannelIndex() == 0);
+	CHECK(ShapeFilter.GetOverlapChannels() == 0);
+	CHECK(ShapeFilter.GetBlockChannels() == 0);
+	CHECK(ShapeFilter.GetMaskFilter() == MaskFilter);
+	CHECK(ShapeFilter.GetFlags() == EmptyFilterFlags);
+}
+
+TEST_CASE("FShapeFilterBuilder - SetFilterFlags", "[Chaos][Filter]")
+{
+	const TArray<EFilterFlags> TestFlags
+	{
+		EFilterFlags::SimpleCollision,
+		EFilterFlags::ComplexCollision,
+		EFilterFlags::SimpleCollision | EFilterFlags::ComplexCollision,
+		EFilterFlags::All,
+	};
+	const uint32 Index = GENERATE(0, 1, 2, 3);
+	const EFilterFlags Flags = TestFlags[Index];
+
+	FShapeFilterBuilder Builder;
+	Builder.SetFilterFlags(Flags);
+	const FShapeFilterData ShapeFilter = Builder.Build();
+
+	CHECK(ShapeFilter.GetCollisionChannelIndex() == 0);
+	CHECK(ShapeFilter.GetOverlapChannels() == 0);
+	CHECK(ShapeFilter.GetBlockChannels() == 0);
+	CHECK(ShapeFilter.GetMaskFilter() == 0);
+	CHECK(ShapeFilter.GetFlags() == Flags);
+}
+
+TEST_CASE("FShapeFilterBuilder - SetFilterFlags with bool", "[Chaos][Filter]")
+{
+	FShapeFilterBuilder Builder;
+	SECTION("Given Empty Flags: Set Single Flag")
+	{
+		Builder.SetFilterFlags(EFilterFlags::ComplexCollision, true);
+		CHECK(EFilterFlags::ComplexCollision == Builder.Build().GetFlags());
+	}
+	SECTION("Given Mixed Flags: Set Multiple Flags ")
+	{
+		Builder.SetFilterFlags(EFilterFlags::SimpleCollision);
+		Builder.SetFilterFlags(EFilterFlags::SimpleCollision | EFilterFlags::ComplexCollision, true);
+		CHECK((EFilterFlags::SimpleCollision | EFilterFlags::ComplexCollision) == Builder.Build().GetFlags());
+	}
+	SECTION("Given Empty Flags: Set All Flags")
+	{
+		Builder.SetFilterFlags(EFilterFlags::All, true);
+		CHECK(EFilterFlags::All == Builder.Build().GetFlags());
+	}
+	SECTION("Given All Flags: Clear Single Flag")
+	{
+		Builder.SetFilterFlags(EFilterFlags::All);
+		Builder.SetFilterFlags(EFilterFlags::ComplexCollision, false);
+		const EFilterFlags ExpectedFlags = EFilterFlags::All & ~EFilterFlags::ComplexCollision;
+		CHECK(ExpectedFlags == Builder.Build().GetFlags());
+	}
+	SECTION("Given Mixed Flags: Clear Multiple Flags")
+	{
+		Builder.SetFilterFlags(EFilterFlags::SimpleCollision);
+		Builder.SetFilterFlags(EFilterFlags::SimpleCollision | EFilterFlags::ComplexCollision, false);
+		CHECK((EFilterFlags::None) == Builder.Build().GetFlags());
+	}
+	SECTION("Given All Flags: Clear All Flags")
+	{
+		Builder.SetFilterFlags(EFilterFlags::All);
+		Builder.SetFilterFlags(EFilterFlags::All, false);
+		CHECK(EFilterFlags::None == Builder.Build().GetFlags());
+	}
+}
+
+TEST_CASE("FShapeFilterBuilder - BuildBlockAll", "[Chaos][Filter]")
+{
+	const EFilterFlags Flags = EFilterFlags::ComplexCollision;
+	const FShapeFilterData ShapeFilter = FShapeFilterBuilder::BuildBlockAll(Flags);
+
+	CHECK(ShapeFilter.GetCollisionChannelIndex() == 0);
+	CHECK(ShapeFilter.GetOverlapChannels() == 0);
+	CHECK(ShapeFilter.GetBlockChannels() == FullChannelMask);
+	CHECK(ShapeFilter.GetMaskFilter() == 0);
+	CHECK(ShapeFilter.GetFlags() == Flags);
+}
+
+TEST_CASE("FQueryFilterBuilder - LegacyRoundTripConstruction", "[Chaos][Filter]")
+{
+	// For each section (flags, collision channel, mask filter) test first bit, last bit, all bits set. Throw in one last test for everything being set.
+	const uint32 Word3 = GENERATE(0, 0x1, 0x80, 0xFF, 1 << 21, 0x10 << 21, 0x1F << 21, 1 << 26, 0x20 << 26, 0x3F << 26, 0xFFE000FF);
+	FCollisionFilterData QueryData;
+	QueryData.Word0 = 1;
+	QueryData.Word1 = 2;
+	QueryData.Word2 = 3;
+	QueryData.Word3 = Word3;
+
+	FQueryFilterData QueryFilterData = FQueryFilterBuilder::BuildFromLegacyQueryFilter(QueryData);
+	const FCollisionFilterData ResultQueryData = FQueryFilterBuilder::GetLegacyQueryFilter(QueryFilterData);
+
+	CHECK(ResultQueryData.Word0 == QueryData.Word0);
+	CHECK(ResultQueryData.Word1 == QueryData.Word1);
+	CHECK(ResultQueryData.Word2 == QueryData.Word2);
+	CHECK(ResultQueryData.Word3 == QueryData.Word3);
+}
+
+TEST_CASE("FQueryObjectFilterBuilder - SetObjectTypes", "[Chaos][Filter]")
+{
+	const TArray<uint64> ObjectTypeMasks{ 0, 0xF, 0xAAAAAAAA, 0xFFFFFFFF, FullChannelMask };
+	const uint32 Index = GENERATE(0, 1, 2, 3);
+	const uint64 ObjectTypeMask = ObjectTypeMasks[Index];
+
+	FQueryObjectFilterBuilder Builder;
+	Builder.SetObjectTypes(ObjectTypeMask);
+	const FQueryFilterData QueryFilterData = Builder.Build();
+
+	CHECK(QueryFilterData.GetQueryType() == FQueryFilterData::EQueryType::ObjectType);
+	CHECK(QueryFilterData.GetObjectTypesToQueryMask() == ObjectTypeMask);
+	CHECK(QueryFilterData.IsMultiQuery() == 0);
+	CHECK(QueryFilterData.GetIgnoreMask() == 0);
+	CHECK(QueryFilterData.GetFlags() == EmptyFilterFlags);
+}
+
+TEST_CASE("FQueryObjectFilterBuilder - SetMultiQuery", "[Chaos][Filter]")
+{
+	const TArray<bool> bMultiQueryStates{ false, true };
+	const uint32 Index = GENERATE(0, 1);
+	const bool bMultiQuery = bMultiQueryStates[Index];
+
+	FQueryObjectFilterBuilder Builder;
+	Builder.SetMultiQuery(bMultiQuery);
+	const FQueryFilterData QueryFilterData = Builder.Build();
+
+
+	CHECK(QueryFilterData.GetQueryType() == FQueryFilterData::EQueryType::ObjectType);
+	CHECK(QueryFilterData.GetObjectTypesToQueryMask() == 0);
+	CHECK(QueryFilterData.IsMultiQuery() == bMultiQuery);
+	CHECK(QueryFilterData.GetIgnoreMask() == 0);
+	CHECK(QueryFilterData.GetFlags() == EmptyFilterFlags);
+}
+
+TEST_CASE("FQueryObjectFilterBuilder - SetMaskFilter", "[Chaos][Filter]")
+{
+	const TArray<uint8> MaskFilters{ 0, 1, 2, 0xF };
+	const uint32 Index = GENERATE(0, 1, 2, 3);
+	const uint8 MaskFilter = MaskFilters[Index];
+
+	FQueryObjectFilterBuilder Builder;
+	Builder.SetMaskFilter(MaskFilter);
+	const FQueryFilterData QueryFilterData = Builder.Build();
+
+	CHECK(QueryFilterData.GetQueryType() == FQueryFilterData::EQueryType::ObjectType);
+	CHECK(QueryFilterData.GetObjectTypesToQueryMask() == 0);
+	CHECK(QueryFilterData.IsMultiQuery() == false);
+	CHECK(QueryFilterData.GetIgnoreMask() == MaskFilter);
+	CHECK(QueryFilterData.GetFlags() == EmptyFilterFlags);
+}
+
+TEST_CASE("FQueryObjectFilterBuilder - SetFilterFlags", "[Chaos][Filter]")
+{
+	const TArray<EFilterFlags> TestFlags
+	{
+		EFilterFlags::SimpleCollision,
+		EFilterFlags::ComplexCollision,
+		EFilterFlags::SimpleCollision | EFilterFlags::ComplexCollision,
+		EFilterFlags::All,
+	};
+	const uint32 Index = GENERATE(0, 1, 2, 3);
+	const EFilterFlags Flags = TestFlags[Index];
+
+	FQueryObjectFilterBuilder Builder;
+	Builder.SetFilterFlags(Flags);
+	const FQueryFilterData QueryFilterData = Builder.Build();
+
+	CHECK(QueryFilterData.GetQueryType() == FQueryFilterData::EQueryType::ObjectType);
+	CHECK(QueryFilterData.GetCollisionChannelIndex() == 0);
+	CHECK(QueryFilterData.GetBlockChannels() == 0);
+	CHECK(QueryFilterData.GetOverlapChannels() == 0);
+	CHECK(QueryFilterData.GetIgnoreMask() == 0);
+	CHECK(QueryFilterData.GetFlags() == Flags);
+}
+
+TEST_CASE("FQueryObjectFilterBuilder - SetFilterFlags with bool", "[Chaos][Filter]")
+{
+	FQueryObjectFilterBuilder Builder;
+	SECTION("Given Empty Flags: Set Single Flag")
+	{
+		Builder.SetFilterFlags(EFilterFlags::ComplexCollision, true);
+		CHECK(EFilterFlags::ComplexCollision == Builder.Build().GetFlags());
+	}
+	SECTION("Given Mixed Flags: Set Multiple Flags ")
+	{
+		Builder.SetFilterFlags(EFilterFlags::SimpleCollision);
+		Builder.SetFilterFlags(EFilterFlags::SimpleCollision | EFilterFlags::ComplexCollision, true);
+		CHECK((EFilterFlags::SimpleCollision | EFilterFlags::ComplexCollision) == Builder.Build().GetFlags());
+	}
+	SECTION("Given Empty Flags: Set All Flags")
+	{
+		Builder.SetFilterFlags(EFilterFlags::All, true);
+		CHECK(EFilterFlags::All == Builder.Build().GetFlags());
+	}
+	SECTION("Given All Flags: Clear Single Flag")
+	{
+		Builder.SetFilterFlags(EFilterFlags::All);
+		Builder.SetFilterFlags(EFilterFlags::ComplexCollision, false);
+		const EFilterFlags ExpectedFlags = EFilterFlags::All & ~EFilterFlags::ComplexCollision;
+		CHECK(ExpectedFlags == Builder.Build().GetFlags());
+	}
+	SECTION("Given Mixed Flags: Clear Multiple Flags")
+	{
+		Builder.SetFilterFlags(EFilterFlags::SimpleCollision);
+		Builder.SetFilterFlags(EFilterFlags::SimpleCollision | EFilterFlags::ComplexCollision, false);
+		CHECK((EFilterFlags::None) == Builder.Build().GetFlags());
+	}
+	SECTION("Given All Flags: Clear All Flags")
+	{
+		Builder.SetFilterFlags(EFilterFlags::All);
+		Builder.SetFilterFlags(EFilterFlags::All, false);
+		CHECK(EFilterFlags::None == Builder.Build().GetFlags());
+	}
+}
+
+TEST_CASE("FQueryObjectFilterBuilder - RoundTrip to legacy and back", "[Chaos][Filter]")
+{
+	const bool bMultiQuery = GENERATE(false, true);
+	FQueryObjectFilterBuilder Builder;
+	Builder.SetObjectTypes(0xA);
+	Builder.SetMultiQuery(bMultiQuery);
+	Builder.SetMaskFilter(0xF);
+	Builder.SetFilterFlags(EFilterFlags::All);
+	const FQueryFilterData ExpectedFilterData = Builder.Build();
+
+	const FCollisionFilterData TempFilterData = FQueryFilterBuilder::GetLegacyQueryFilter(ExpectedFilterData);
+	const FQueryFilterData ResultFilterData = FQueryFilterBuilder::BuildFromLegacyQueryFilter(TempFilterData);
+
+	CHECK(ExpectedFilterData.GetQueryType() == ResultFilterData.GetQueryType());
+	CHECK(ExpectedFilterData.GetObjectTypesToQueryMask() == ResultFilterData.GetObjectTypesToQueryMask());
+	CHECK(ExpectedFilterData.IsMultiQuery() == ResultFilterData.IsMultiQuery());
+	CHECK(ExpectedFilterData.GetIgnoreMask() == ResultFilterData.GetIgnoreMask());
+	CHECK(ExpectedFilterData.GetFlags() == ResultFilterData.GetFlags());
+}
+
+TEST_CASE("FQueryTraceFilterBuilder - SetCollisionChannelIndex", "[Chaos][Filter]")
+{
+	const uint8 ChannelIndex = GENERATE(0, 1, 3, 31, 32, 63);
+
+	FQueryTraceFilterBuilder Builder;
+	Builder.SetCollisionChannelIndex(ChannelIndex);
+	const FQueryFilterData QueryFilterData = Builder.Build();
+
+	CHECK(QueryFilterData.GetQueryType() == FQueryFilterData::EQueryType::Channel);
+	CHECK(QueryFilterData.GetCollisionChannelIndex() == ChannelIndex);
+	CHECK(QueryFilterData.GetCollisionChannelMask() == ToChannelMask(ChannelIndex));
+	CHECK(QueryFilterData.GetBlockChannels() == 0);
+	CHECK(QueryFilterData.GetOverlapChannels() == 0);
+	CHECK(QueryFilterData.GetIgnoreMask() == 0);
+	CHECK(QueryFilterData.GetFlags() == EmptyFilterFlags);
+}
+
+TEST_CASE("FQueryTraceFilterBuilder - SetBlockChannelMask", "[Chaos][Filter]")
+{
+	const TArray<uint64> BlockChannelMasks{ 0, 1, 2, 3, 0xAAAAAAAA, 0xFFFFFFFF, FullChannelMask };
+	const uint32 Index = GENERATE_COPY(Catch::Generators::range(0, BlockChannelMasks.Num()));
+	const uint64 BlockChannelMask = BlockChannelMasks[Index];
+
+	FQueryTraceFilterBuilder Builder;
+	Builder.SetBlockChannelMask(BlockChannelMask);
+	const FQueryFilterData QueryFilterData = Builder.Build();
+
+	CHECK(QueryFilterData.GetQueryType() == FQueryFilterData::EQueryType::Channel);
+	CHECK(QueryFilterData.GetCollisionChannelIndex() == 0);
+	CHECK(QueryFilterData.GetBlockChannels() == BlockChannelMask);
+	CHECK(QueryFilterData.GetOverlapChannels() == 0);
+	CHECK(QueryFilterData.GetIgnoreMask() == 0);
+	CHECK(QueryFilterData.GetFlags() == EmptyFilterFlags);
+}
+
+TEST_CASE("FQueryTraceFilterBuilder - SetOverlapChannelMask", "[Chaos][Filter]")
+{
+	const TArray<uint64> OverlapChannelMasks{ 0, 1, 2, 3, 0xAAAAAAAA, 0xFFFFFFFF, FullChannelMask };
+	const uint32 Index = GENERATE_COPY(Catch::Generators::range(0, OverlapChannelMasks.Num()));
+	const uint64 OverlapChannelMask = OverlapChannelMasks[Index];
+
+	FQueryTraceFilterBuilder Builder;
+	Builder.SetOverlapChannelMask(OverlapChannelMask);
+	const FQueryFilterData QueryFilterData = Builder.Build();
+
+	CHECK(QueryFilterData.GetQueryType() == FQueryFilterData::EQueryType::Channel);
+	CHECK(QueryFilterData.GetCollisionChannelIndex() == 0);
+	CHECK(QueryFilterData.GetBlockChannels() == 0);
+	CHECK(QueryFilterData.GetOverlapChannels() == OverlapChannelMask);
+	CHECK(QueryFilterData.GetIgnoreMask() == 0);
+	CHECK(QueryFilterData.GetFlags() == EmptyFilterFlags);
+}
+
+TEST_CASE("FQueryTraceFilterBuilder - SetMaskFilter", "[Chaos][Filter]")
+{
+	const TArray<uint8> MaskFilters{ 0, 1, 2, 0xF };
+	const uint32 Index = GENERATE(0, 1, 2, 3);
+	const uint8 MaskFilter = MaskFilters[Index];
+
+	FQueryTraceFilterBuilder Builder;
+	Builder.SetMaskFilter(MaskFilter);
+	const FQueryFilterData QueryFilterData = Builder.Build();
+
+	CHECK(QueryFilterData.GetQueryType() == FQueryFilterData::EQueryType::Channel);
+	CHECK(QueryFilterData.GetCollisionChannelIndex() == 0);
+	CHECK(QueryFilterData.GetBlockChannels() == 0);
+	CHECK(QueryFilterData.GetOverlapChannels() == 0);
+	CHECK(QueryFilterData.GetIgnoreMask() == MaskFilter);
+	CHECK(QueryFilterData.GetFlags() == EmptyFilterFlags);
+}
+
+TEST_CASE("FQueryTraceFilterBuilder - SetFilterFlags", "[Chaos][Filter]")
+{
+	const TArray<EFilterFlags> TestFlags
+	{
+		EFilterFlags::SimpleCollision,
+		EFilterFlags::ComplexCollision,
+		EFilterFlags::SimpleCollision | EFilterFlags::ComplexCollision,
+		EFilterFlags::All,
+	};
+	const uint32 Index = GENERATE(0, 1, 2, 3);
+	const EFilterFlags Flags = TestFlags[Index];
+
+	FQueryTraceFilterBuilder Builder;
+	Builder.SetFilterFlags(Flags);
+	const FQueryFilterData QueryFilterData = Builder.Build();
+
+	CHECK(QueryFilterData.GetQueryType() == FQueryFilterData::EQueryType::Channel);
+	CHECK(QueryFilterData.GetCollisionChannelIndex() == 0);
+	CHECK(QueryFilterData.GetBlockChannels() == 0);
+	CHECK(QueryFilterData.GetOverlapChannels() == 0);
+	CHECK(QueryFilterData.GetIgnoreMask() == 0);
+	CHECK(QueryFilterData.GetFlags() == Flags);
+}
+
+TEST_CASE("FQueryTraceFilterBuilder - SetFilterFlags with bool", "[Chaos][Filter]")
+{
+	FQueryTraceFilterBuilder Builder;
+	SECTION("Given Empty Flags: Set Single Flag")
+	{
+		Builder.SetFilterFlags(EFilterFlags::ComplexCollision, true);
+		CHECK(EFilterFlags::ComplexCollision == Builder.Build().GetFlags());
+	}
+	SECTION("Given Mixed Flags: Set Multiple Flags ")
+	{
+		Builder.SetFilterFlags(EFilterFlags::SimpleCollision);
+		Builder.SetFilterFlags(EFilterFlags::SimpleCollision | EFilterFlags::ComplexCollision, true);
+		CHECK((EFilterFlags::SimpleCollision | EFilterFlags::ComplexCollision) == Builder.Build().GetFlags());
+	}
+	SECTION("Given Empty Flags: Set All Flags")
+	{
+		Builder.SetFilterFlags(EFilterFlags::All, true);
+		CHECK(EFilterFlags::All == Builder.Build().GetFlags());
+	}
+	SECTION("Given All Flags: Clear Single Flag")
+	{
+		Builder.SetFilterFlags(EFilterFlags::All);
+		Builder.SetFilterFlags(EFilterFlags::ComplexCollision, false);
+		const EFilterFlags ExpectedFlags = EFilterFlags::All & ~EFilterFlags::ComplexCollision;
+		CHECK(ExpectedFlags == Builder.Build().GetFlags());
+	}
+	SECTION("Given Mixed Flags: Clear Multiple Flags")
+	{
+		Builder.SetFilterFlags(EFilterFlags::SimpleCollision);
+		Builder.SetFilterFlags(EFilterFlags::SimpleCollision | EFilterFlags::ComplexCollision, false);
+		CHECK((EFilterFlags::None) == Builder.Build().GetFlags());
+	}
+	SECTION("Given All Flags: Clear All Flags")
+	{
+		Builder.SetFilterFlags(EFilterFlags::All);
+		Builder.SetFilterFlags(EFilterFlags::All, false);
+		CHECK(EFilterFlags::None == Builder.Build().GetFlags());
+	}
+}
+
+TEST_CASE("FQueryTraceFilterBuilder - BuildOverlapAll", "[Chaos][Filter]")
+{
+	const EFilterFlags Flags = EFilterFlags::ComplexCollision;
+	const FQueryFilterData QueryFilterData = FQueryTraceFilterBuilder::BuildOverlapAll(Flags);
+
+	CHECK(QueryFilterData.GetQueryType() == FQueryFilterData::EQueryType::Channel);
+	CHECK(QueryFilterData.GetCollisionChannelIndex() == 0);
+	CHECK(QueryFilterData.GetBlockChannels() == 0);
+	CHECK(QueryFilterData.GetOverlapChannels() == FullChannelMask);
+	CHECK(QueryFilterData.GetIgnoreMask() == 0);
+	CHECK(QueryFilterData.GetFlags() == Flags);
+}
+
+TEST_CASE("FQueryTraceFilterBuilder - RoundTrip to legacy and back", "[Chaos][Filter]")
+{
+	const uint32 ChannelIndex = GENERATE(0, 1, 31);
+	FQueryTraceFilterBuilder Builder;
+	Builder.SetCollisionChannelIndex(ChannelIndex);
+	Builder.SetBlockChannelMask(2);
+	Builder.SetOverlapChannelMask(3);
+	Builder.SetMaskFilter(0xF);
+	Builder.SetFilterFlags(EFilterFlags::All);
+	const FQueryFilterData ExpectedFilterData = Builder.Build();
+
+	const FCollisionFilterData TempFilterData = FQueryFilterBuilder::GetLegacyQueryFilter(ExpectedFilterData);
+	const FQueryFilterData ResultFilterData = FQueryFilterBuilder::BuildFromLegacyQueryFilter(TempFilterData);
+
+	CHECK(ExpectedFilterData.GetQueryType() == ResultFilterData.GetQueryType());
+	CHECK(ExpectedFilterData.GetCollisionChannelIndex() == ResultFilterData.GetCollisionChannelIndex());
+	CHECK(ExpectedFilterData.GetCollisionChannelMask() == ResultFilterData.GetCollisionChannelMask());
+	CHECK(ExpectedFilterData.GetBlockChannels() == ResultFilterData.GetBlockChannels());
+	CHECK(ExpectedFilterData.GetOverlapChannels() == ResultFilterData.GetOverlapChannels());
+	CHECK(ExpectedFilterData.GetIgnoreMask() == ResultFilterData.GetIgnoreMask());
+	CHECK(ExpectedFilterData.GetFlags() == ResultFilterData.GetFlags());
+}
+
+TEST_CASE("FShapeFilterData - NarrowFilter", "[Chaos][Filter]")
+{
+	auto BuildFiltersAndRunNarrowFilter = [](const uint8 CollisionChannelIndex0, const uint64 BlockMask0, const uint8 CollisionChannelIndex1, const uint64 BlockMask1)
+	{
+		const FShapeFilterData Filter0 = FShapeFilterBuilder().SetCollisionChannelIndex(CollisionChannelIndex0).SetBlockChannelMask(BlockMask0).Build();
+		const FShapeFilterData Filter1 = FShapeFilterBuilder().SetCollisionChannelIndex(CollisionChannelIndex1).SetBlockChannelMask(BlockMask1).Build();
+		return Filter0.NarrowFilter(Filter1);
+	};
+
+	// Test the block masks not matching
+	CHECK(ENarrowFilterResult::None == BuildFiltersAndRunNarrowFilter(0, 0b01, 1, 0b01));
+	CHECK(ENarrowFilterResult::None == BuildFiltersAndRunNarrowFilter(0, 0b10, 1, 0b10));
+
+	// Test each bit position
+	for (uint8 I = 0; I < MaxChannelBitCount; ++I)
+	{
+		const uint64 ChannelBit = ToChannelMask(I);
+		CHECK(ENarrowFilterResult::Block == BuildFiltersAndRunNarrowFilter(0, ChannelBit, I, 0b01));
+		CHECK(ENarrowFilterResult::Block == BuildFiltersAndRunNarrowFilter(I, 0b01, 0, ChannelBit));
+		CHECK(ENarrowFilterResult::Block == BuildFiltersAndRunNarrowFilter(0, FullChannelMask, I, 0b01));
+		CHECK(ENarrowFilterResult::Block == BuildFiltersAndRunNarrowFilter(I, 0b01, 0, FullChannelMask));
+	}
+
+	// Test empty block masks
+	CHECK(ENarrowFilterResult::None == BuildFiltersAndRunNarrowFilter(0, 0b00, 1, 0b00));
+	CHECK(ENarrowFilterResult::None == BuildFiltersAndRunNarrowFilter(0, 0b10, 1, 0b00));
+	CHECK(ENarrowFilterResult::None == BuildFiltersAndRunNarrowFilter(0, 0b00, 1, 0b01));
+
+	// Test all permutations of 4 bits to get a good sample size
+	CHECK(ENarrowFilterResult::None == BuildFiltersAndRunNarrowFilter(0, 0b1010, 0, 0b0101));
+	CHECK(ENarrowFilterResult::Block == BuildFiltersAndRunNarrowFilter(0, 0b1010, 1, 0b0101));
+	CHECK(ENarrowFilterResult::None == BuildFiltersAndRunNarrowFilter(0, 0b1010, 2, 0b0101));
+	CHECK(ENarrowFilterResult::Block == BuildFiltersAndRunNarrowFilter(0, 0b1010, 3, 0b0101));
+
+	CHECK(ENarrowFilterResult::None == BuildFiltersAndRunNarrowFilter(1, 0b1010, 0, 0b0101));
+	CHECK(ENarrowFilterResult::None == BuildFiltersAndRunNarrowFilter(1, 0b1010, 1, 0b0101));
+	CHECK(ENarrowFilterResult::None == BuildFiltersAndRunNarrowFilter(1, 0b1010, 2, 0b0101));
+	CHECK(ENarrowFilterResult::None == BuildFiltersAndRunNarrowFilter(1, 0b1010, 3, 0b0101));
+
+	CHECK(ENarrowFilterResult::None == BuildFiltersAndRunNarrowFilter(2, 0b1010, 0, 0b0101));
+	CHECK(ENarrowFilterResult::Block == BuildFiltersAndRunNarrowFilter(2, 0b1010, 1, 0b0101));
+	CHECK(ENarrowFilterResult::None == BuildFiltersAndRunNarrowFilter(2, 0b1010, 2, 0b0101));
+	CHECK(ENarrowFilterResult::Block == BuildFiltersAndRunNarrowFilter(2, 0b1010, 3, 0b0101));
+
+	CHECK(ENarrowFilterResult::None == BuildFiltersAndRunNarrowFilter(3, 0b1010, 0, 0b0101));
+	CHECK(ENarrowFilterResult::None == BuildFiltersAndRunNarrowFilter(3, 0b1010, 1, 0b0101));
+	CHECK(ENarrowFilterResult::None == BuildFiltersAndRunNarrowFilter(3, 0b1010, 2, 0b0101));
+	CHECK(ENarrowFilterResult::None == BuildFiltersAndRunNarrowFilter(3, 0b1010, 3, 0b0101));
+}
+
+TEST_CASE("FQueryFilterData - NarrowFilter - ObjectType", "[Chaos][Filter]")
+{
+	auto BuildAndRunObjectTypeNarrowFilter = [](const uint8 ShapeCollisionChannel, const uint64 ObjectTypeMask)
+	{
+		const FShapeFilterData ShapeFilter = FShapeFilterBuilder().SetCollisionChannelIndex(ShapeCollisionChannel).Build();
+		const FQueryFilterData QueryFilter = FQueryObjectFilterBuilder().SetObjectTypes(ObjectTypeMask).Build();
+		return QueryFilter.NarrowFilter(ShapeFilter);
+	};
+	auto BuildAndRunObjectTypeNarrowFilterWithMasks = [](const uint8 CollisionChannel, const uint8 ShapeMaskFilter, const uint8 QueryMaskFilter)
+	{
+		const FShapeFilterData ShapeFilter = FShapeFilterBuilder().SetCollisionChannelIndex(CollisionChannel).SetMaskFilter(ShapeMaskFilter).Build();
+		const FQueryFilterData QueryFilter = FQueryObjectFilterBuilder().SetObjectTypes(ToChannelMask(CollisionChannel)).SetMaskFilter(QueryMaskFilter).Build();
+		return QueryFilter.NarrowFilter(ShapeFilter);
+
+	};
+	for (uint8 I = 0; I < MaxChannelBitCount; ++I)
+	{
+		CHECK(ENarrowFilterResult::None == BuildAndRunObjectTypeNarrowFilter(I, 0));
+		CHECK(ENarrowFilterResult::Block == BuildAndRunObjectTypeNarrowFilter(I, ToChannelMask(I)));
+		CHECK(ENarrowFilterResult::Block == BuildAndRunObjectTypeNarrowFilter(I, FullChannelMask));
+	}
+
+	CHECK(ENarrowFilterResult::None == BuildAndRunObjectTypeNarrowFilter(0, 0b1010));
+	CHECK(ENarrowFilterResult::Block == BuildAndRunObjectTypeNarrowFilter(1, 0b1010));
+	CHECK(ENarrowFilterResult::None == BuildAndRunObjectTypeNarrowFilter(2, 0b1010));
+	CHECK(ENarrowFilterResult::Block == BuildAndRunObjectTypeNarrowFilter(3, 0b1010));
+	CHECK(ENarrowFilterResult::None == BuildAndRunObjectTypeNarrowFilter(4, 0b1010));
+
+	// Verify the mask filter
+	CHECK(ENarrowFilterResult::Block == BuildAndRunObjectTypeNarrowFilterWithMasks(3, 0b0110, 0b1001));
+	CHECK(ENarrowFilterResult::None == BuildAndRunObjectTypeNarrowFilterWithMasks(3, 0b0110, 0b1100));
+}
+
+TEST_CASE("FQueryFilterData - NarrowFilter - Channel", "[Chaos][Filter]")
+{
+	auto BuildAndRunChannelNarrowFilter = [](const uint8 ShapeChannel, const uint64 ShapeOverlap, const uint64 ShapeBlock, const uint8 QueryChannel, const uint64 QueryOverlap, const uint64 QueryBlock)
+	{
+		const FShapeFilterData ShapeFilter = FShapeFilterBuilder()
+			.SetCollisionChannelIndex(ShapeChannel)
+			.SetOverlapChannelMask(ShapeOverlap)
+			.SetBlockChannelMask(ShapeBlock)
+			.Build();
+		const FQueryFilterData QueryFilter = FQueryTraceFilterBuilder()
+			.SetCollisionChannelIndex(QueryChannel)
+			.SetOverlapChannelMask(QueryOverlap)
+			.SetBlockChannelMask(QueryBlock)
+			.Build();
+		return QueryFilter.NarrowFilter(ShapeFilter);
+	};
+	auto BuildAndRunChannelTypeNarrowFilterWithMasks = [](const uint8 ChannelIndex, const uint8 ShapeMaskFilter, const uint8 QueryMaskFilter)
+	{
+		const uint64 BlockMask = ToChannelMask(ChannelIndex);
+		const FShapeFilterData ShapeFilter = FShapeFilterBuilder()
+			.SetCollisionChannelIndex(ChannelIndex)
+			.SetBlockChannelMask(BlockMask)
+			.SetMaskFilter(ShapeMaskFilter)
+			.Build();
+		const FQueryFilterData QueryFilter = FQueryTraceFilterBuilder()
+			.SetCollisionChannelIndex(ChannelIndex)
+			.SetBlockChannelMask(BlockMask)
+			.SetMaskFilter(QueryMaskFilter)
+			.Build();
+		return QueryFilter.NarrowFilter(ShapeFilter);
+	};
+
+	// Rules: 
+	// - For a single check: Block > Overlap > None (e.g. Block & Overlap -> Block)
+	// - For a combined check: None > Overlap > Block (e.g. Block vs. Overlap -> Overlap)
+
+	constexpr uint64 FullMask = FullChannelMask;
+	for (uint8 I = 0; I < MaxChannelBitCount; ++I)
+	{
+		const uint64 ChannelBit = ToChannelMask(I);
+		// (None, None) -> None
+		CHECK(ENarrowFilterResult::None == BuildAndRunChannelNarrowFilter(I, 0, 0, I, 0, 0));
+		// (None, Overlap) -> None
+		CHECK(ENarrowFilterResult::None == BuildAndRunChannelNarrowFilter(I, 0, 0, I, ChannelBit, 0));
+		CHECK(ENarrowFilterResult::None == BuildAndRunChannelNarrowFilter(I, 0, 0, I, FullMask, 0));
+		// (None, Block) -> None
+		CHECK(ENarrowFilterResult::None == BuildAndRunChannelNarrowFilter(I, 0, 0, I, 0, ChannelBit));
+		CHECK(ENarrowFilterResult::None == BuildAndRunChannelNarrowFilter(I, 0, 0, I, 0, FullMask));
+		// (None, Overlap & Block) -> None
+		CHECK(ENarrowFilterResult::None == BuildAndRunChannelNarrowFilter(I, 0, 0, I, ChannelBit, ChannelBit));
+		CHECK(ENarrowFilterResult::None == BuildAndRunChannelNarrowFilter(I, 0, 0, I, FullMask, FullMask));
+
+		// (Overlap, None) -> None
+		CHECK(ENarrowFilterResult::None == BuildAndRunChannelNarrowFilter(I, ChannelBit, 0, I, 0, 0));
+		CHECK(ENarrowFilterResult::None == BuildAndRunChannelNarrowFilter(I, FullMask, 0, I, 0, 0));
+		// (Overlap, Overlap) -> Overlap
+		CHECK(ENarrowFilterResult::Overlap == BuildAndRunChannelNarrowFilter(I, ChannelBit, 0, I, ChannelBit, 0));
+		CHECK(ENarrowFilterResult::Overlap == BuildAndRunChannelNarrowFilter(I, FullMask, 0, I, FullMask, 0));
+		// (Overlap, Block) -> Overlap
+		CHECK(ENarrowFilterResult::Overlap == BuildAndRunChannelNarrowFilter(I, ChannelBit, 0, I, 0, ChannelBit));
+		CHECK(ENarrowFilterResult::Overlap == BuildAndRunChannelNarrowFilter(I, FullMask, 0, I, 0, FullMask));
+		// (Overlap, Overlap & Block) -> Overlap
+		CHECK(ENarrowFilterResult::Overlap == BuildAndRunChannelNarrowFilter(I, ChannelBit, 0, I, ChannelBit, ChannelBit));
+		CHECK(ENarrowFilterResult::Overlap == BuildAndRunChannelNarrowFilter(I, FullMask, 0, I, FullMask, FullMask));
+
+		// (Block, None) -> None
+		CHECK(ENarrowFilterResult::None == BuildAndRunChannelNarrowFilter(I, 0, ChannelBit, I, 0, 0));
+		CHECK(ENarrowFilterResult::None == BuildAndRunChannelNarrowFilter(I, 0, FullMask, I, 0, 0));
+		// (Block, Overlap) -> Overlap
+		CHECK(ENarrowFilterResult::Overlap == BuildAndRunChannelNarrowFilter(I, 0, ChannelBit, I, ChannelBit, 0));
+		CHECK(ENarrowFilterResult::Overlap == BuildAndRunChannelNarrowFilter(I, 0, FullMask, I, FullMask, 0));
+		// (Block, Block) -> Block
+		CHECK(ENarrowFilterResult::Block == BuildAndRunChannelNarrowFilter(I, 0, ChannelBit, I, 0, ChannelBit));
+		CHECK(ENarrowFilterResult::Block == BuildAndRunChannelNarrowFilter(I, 0, FullMask, I, 0, FullMask));
+		// (Block, Overlap & Block) -> Block
+		CHECK(ENarrowFilterResult::Block == BuildAndRunChannelNarrowFilter(I, 0, ChannelBit, I, ChannelBit, ChannelBit));// Broken??
+		CHECK(ENarrowFilterResult::Block == BuildAndRunChannelNarrowFilter(I, 0, FullMask, I, FullMask, FullMask));
+
+		// (Overlap & Block, Overlap & Block) -> Block
+		CHECK(ENarrowFilterResult::Block == BuildAndRunChannelNarrowFilter(I, ChannelBit, ChannelBit, I, ChannelBit, ChannelBit));
+		CHECK(ENarrowFilterResult::Block == BuildAndRunChannelNarrowFilter(I, FullMask, FullMask, I, FullMask, FullMask));
+	}
+
+	// Test individual bit patterns...
+
+	// Block vs.
+	CHECK(ENarrowFilterResult::Block == BuildAndRunChannelNarrowFilter(0, 0b10100000, 0b01010000, 4, 0b00001010, 0b00000101));
+	CHECK(ENarrowFilterResult::Overlap == BuildAndRunChannelNarrowFilter(1, 0b10100000, 0b01010000, 4, 0b00001010, 0b00000101));
+	// Overlap vs.
+	CHECK(ENarrowFilterResult::Overlap == BuildAndRunChannelNarrowFilter(2, 0b10100000, 0b01010000, 5, 0b00001010, 0b00000101));
+	CHECK(ENarrowFilterResult::Overlap == BuildAndRunChannelNarrowFilter(3, 0b10100000, 0b01010000, 5, 0b00001010, 0b00000101));
+	// Block vs.
+	CHECK(ENarrowFilterResult::Block == BuildAndRunChannelNarrowFilter(0, 0b10100000, 0b01010000, 6, 0b00001010, 0b00000101));
+	CHECK(ENarrowFilterResult::Overlap == BuildAndRunChannelNarrowFilter(1, 0b10100000, 0b01010000, 6, 0b00001010, 0b00000101));
+	// Overlap vs.
+	CHECK(ENarrowFilterResult::Overlap == BuildAndRunChannelNarrowFilter(2, 0b10100000, 0b01010000, 7, 0b00001010, 0b00000101));
+	CHECK(ENarrowFilterResult::Overlap == BuildAndRunChannelNarrowFilter(3, 0b10100000, 0b01010000, 7, 0b00001010, 0b00000101));
+	// Test edge boundaries for none
+	CHECK(ENarrowFilterResult::None == BuildAndRunChannelNarrowFilter(0, 0b10100000, 0b01010000, 3, 0b00001010, 0b00000101));
+	CHECK(ENarrowFilterResult::None == BuildAndRunChannelNarrowFilter(0, 0b10100000, 0b01010000, 8, 0b00001010, 0b00000101));
+	CHECK(ENarrowFilterResult::None == BuildAndRunChannelNarrowFilter(5, 0b10100000, 0b01010000, 4, 0b00001010, 0b00000101));
+
+	// Verify the mask filter
+	CHECK(ENarrowFilterResult::Block == BuildAndRunChannelTypeNarrowFilterWithMasks(3, 0b0110, 0b1001));
+	CHECK(ENarrowFilterResult::None == BuildAndRunChannelTypeNarrowFilterWithMasks(3, 0b0110, 0b1100));
+}
+
+TEST_CASE("FShapeUnionFilterData - BroadFilter - ShapeUnion", "[Chaos][Filter]")
+{
+	// Single
+	SECTION("Single - Accept")
+	{
+		const FShapeUnionFilterData Union0 = BuildUnionFilter({ BuildShapeFilter(0, 0b10, 0) });
+		const FShapeUnionFilterData Union1 = BuildUnionFilter({ BuildShapeFilter(1, 0b01, 0) });
+		CHECK(EBroadFilterResult::Accept == Union0.BroadFilter(Union1));
+	}
+	SECTION("Single - Reject")
+	{
+		const FShapeUnionFilterData Union0 = BuildUnionFilter({ BuildShapeFilter(0, 0b10, 0) });
+		const FShapeUnionFilterData Union1 = BuildUnionFilter({ BuildShapeFilter(1, 0b10, 0) });
+		CHECK(EBroadFilterResult::Reject == Union0.BroadFilter(Union1));
+	}
+	SECTION("Multiple - Both Reject")
+	{
+		// A { Channel(0b0101) Block(0b1010) } B { Channel(0b0101) Block(0b1010) }
+		const FShapeUnionFilterData Union0 = BuildUnionFilter({ BuildShapeFilter(0, 0b0010, 0), BuildShapeFilter(2, 0b1000, 0) });
+		const FShapeUnionFilterData Union1 = BuildUnionFilter({ BuildShapeFilter(0, 0b0010, 0), BuildShapeFilter(2, 0b1000, 0) });
+		CHECK(EBroadFilterResult::Reject == Union0.BroadFilter(Union1));
+	}
+	SECTION("Multiple - One Rejects")
+	{
+		// A { Channel(0b0101) Block(0b1010) } B { Channel(0b0101) Block(0b0011) }
+		const FShapeUnionFilterData Union0 = BuildUnionFilter({ BuildShapeFilter(0, 0b0010, 0), BuildShapeFilter(2, 0b1000, 0) });
+		const FShapeUnionFilterData Union1 = BuildUnionFilter({ BuildShapeFilter(0, 0b0010, 0), BuildShapeFilter(2, 0b0001, 0) });
+		CHECK(EBroadFilterResult::Reject == Union0.BroadFilter(Union1));
+		CHECK(EBroadFilterResult::Reject == Union1.BroadFilter(Union0));
+	}
+	SECTION("Multiple - Accept")
+	{
+		// A { Channel(0b0011) Block(0b1100) } B { Channel(0b1100) Block(0b0011) }
+		const FShapeUnionFilterData Union0 = BuildUnionFilter({ BuildShapeFilter(0, 0b1000, 0), BuildShapeFilter(1, 0b0100, 0) });
+		const FShapeUnionFilterData Union1 = BuildUnionFilter({ BuildShapeFilter(2, 0b0010, 0), BuildShapeFilter(3, 0b0001, 0) });
+		CHECK(EBroadFilterResult::Accept == Union0.BroadFilter(Union1));
+	}
+	SECTION("Combine No Sim")
+	{
+		const FShapeFilterData ShapeFilter0 = BuildShapeFilter(0, 0b01, 0);
+		const FShapeFilterData ShapeFilter1 = BuildShapeFilter(1, 0b10, 0);
+
+		const FShapeUnionFilterData Union0 = BuildUnionFilter({ ShapeFilter0 });
+		const FShapeUnionFilterData Union1 = BuildUnionFilter({ ShapeFilter1 });
+		FShapeUnionFilterData Union01;
+		Union01.Combine(ShapeFilter0, true, true);
+		Union01.Combine(ShapeFilter1, true, false);
+
+		CHECK(EBroadFilterResult::Accept == Union01.BroadFilter(Union0));
+		CHECK(EBroadFilterResult::Reject == Union01.BroadFilter(Union1));
+	}
+}
+
+TEST_CASE("FShapeUnionFilterData - BroadFilter - Query", "[Chaos][Filter]")
+{
+	const uint8 ChannelIndex = 0;
+
+	SECTION("QueryType(Channel)")
+	{
+		const TArray<FQueryFilterData> QueryFilters
+		{
+			FQueryTraceFilterBuilder().SetCollisionChannelIndex(0).Build(),
+			FQueryTraceFilterBuilder().SetCollisionChannelIndex(1).Build(),
+			FQueryTraceFilterBuilder().SetCollisionChannelIndex(2).Build(),
+			FQueryTraceFilterBuilder().SetCollisionChannelIndex(3).Build(),
+		};
+		SECTION("Empty Filters")
+		{
+			const FQueryFilterData QueryFilter = FQueryTraceFilterBuilder().Build();
+			FShapeUnionFilterData Union0;
+			CHECK(EBroadFilterResult::Reject == Union0.BroadFilter(QueryFilter));
+
+			FShapeUnionFilterData Union1;
+			Union1.Combine(FShapeFilterBuilder().Build());
+			CHECK(EBroadFilterResult::Reject == Union1.BroadFilter(QueryFilter));
+		}
+		SECTION("Block Masks")
+		{
+			FShapeUnionFilterData Union;
+			Union.Combine(BuildShapeFilter(ChannelIndex, 0b0001, 0));
+			Union.Combine(BuildShapeFilter(ChannelIndex, 0b0100, 0));
+
+			CHECK(EBroadFilterResult::Accept == Union.BroadFilter(QueryFilters[0]));
+			CHECK(EBroadFilterResult::Reject == Union.BroadFilter(QueryFilters[1]));
+			CHECK(EBroadFilterResult::Accept == Union.BroadFilter(QueryFilters[2]));
+			CHECK(EBroadFilterResult::Reject == Union.BroadFilter(QueryFilters[3]));
+		}
+		SECTION("Overlap Masks")
+		{
+			FShapeUnionFilterData Union;
+			Union.Combine(BuildShapeFilter(ChannelIndex, 0, 0b0001));
+			Union.Combine(BuildShapeFilter(ChannelIndex, 0, 0b0100));
+
+			CHECK(EBroadFilterResult::Accept == Union.BroadFilter(QueryFilters[0]));
+			CHECK(EBroadFilterResult::Reject == Union.BroadFilter(QueryFilters[1]));
+			CHECK(EBroadFilterResult::Accept == Union.BroadFilter(QueryFilters[2]));
+			CHECK(EBroadFilterResult::Reject == Union.BroadFilter(QueryFilters[3]));
+		}
+		SECTION("Block and Overlap Masks")
+		{
+			FShapeUnionFilterData Union;
+			// Note: Bit 2 is true for both
+			Union.Combine(BuildShapeFilter(ChannelIndex, 0b0001, 0b0000));
+			Union.Combine(BuildShapeFilter(ChannelIndex, 0b0010, 0b0000));
+			Union.Combine(BuildShapeFilter(ChannelIndex, 0b0000, 0b0010));
+			Union.Combine(BuildShapeFilter(ChannelIndex, 0b0000, 0b1000));
+
+			CHECK(EBroadFilterResult::Accept == Union.BroadFilter(QueryFilters[0]));
+			CHECK(EBroadFilterResult::Accept == Union.BroadFilter(QueryFilters[1]));
+			CHECK(EBroadFilterResult::Reject == Union.BroadFilter(QueryFilters[2]));
+			CHECK(EBroadFilterResult::Accept == Union.BroadFilter(QueryFilters[3]));
+		}
+		SECTION("Combine No Sim")
+		{
+			FShapeUnionFilterData Union;
+			Union.Combine(BuildShapeFilter(ChannelIndex, 0, 0b01));
+			Union.Combine(BuildShapeFilter(ChannelIndex, 1, 0b10), false);
+
+			CHECK(EBroadFilterResult::Accept == Union.BroadFilter(QueryFilters[0]));
+			CHECK(EBroadFilterResult::Reject == Union.BroadFilter(QueryFilters[1]));
+		}
+	}
+
+	SECTION("QueryType(ObjectType)")
+	{
+		// Note: ObjectType filters do no filtering now
+		const FQueryFilterData QueryFilter0 = FQueryObjectFilterBuilder().Build();
+		const FQueryFilterData QueryFilter1 = FQueryObjectFilterBuilder().SetObjectTypes(0b01).Build();
+		const FQueryFilterData QueryFilter2 = FQueryObjectFilterBuilder().SetObjectTypes(0b10).Build();
+
+		FShapeUnionFilterData Union;
+		Union.Combine(BuildShapeFilter(ChannelIndex, 0b0001, 0b0001));
+
+		CHECK(EBroadFilterResult::Accept == Union.BroadFilter(QueryFilter0));
+		CHECK(EBroadFilterResult::Accept == Union.BroadFilter(QueryFilter1));
+		CHECK(EBroadFilterResult::Accept == Union.BroadFilter(QueryFilter2));
+	}
+}
+
+} // namespace Chaos::Filter
diff --git a/Engine/Source/Programs/ChaosPhysicsTests/Private/FrictionTests.cpp b/Engine/Source/Programs/ChaosPhysicsTests/Private/FrictionTests.cpp
new file mode 100644
--- /dev/null
+++ b/Engine/Source/Programs/ChaosPhysicsTests/Private/FrictionTests.cpp
@@ -0,0 +1,86 @@
+// Copyright Epic Games, Inc. All Rights Reserved.
+#include "CoreMinimal.h"
+#include "Chaos/Box.h"
+#include "Chaos/Collision/PBDCollisionConstraint.h"
+#include "Chaos/ContactModification.h"
+#include "Chaos/MassProperties.h"
+#include "Chaos/PBDCollisionConstraints.h"
+#include "Chaos/PBDRigidsEvolution.h"
+#include "Chaos/SimCallbackObject.h"
+#include "ChaosSolversModule.h"
+#include "ChaosTestHarness.h"
+#include "ChaosTestScene.h"
+#include "EventsData.h"
+#include "PBDRigidsSolver.h"
+#include "PhysicsProxy/SingleParticlePhysicsProxy.h"
+#include "ScopedCVar.h"
+
+//
+// Test initial friction by dropping a moving ball onto a plane with enough 
+// friction to stop the contact point.
+// The results should be independent of initial penetration depth.
+// The results should be independent of friction as long as it is above the sticking limit.
+//
+TEST_CASE("ChaosImpactFrictionTests", "[Chaos][Friction][unit]")
+{
+	using namespace Chaos;
+	using namespace Chaos::LowLevelTest;
+
+	const float Dt = 1.0f / 30.0f;
+	const float G = 1000.0f;
+
+	FChaosTestScene Scene;
+	Scene.SetGravity(FVector3d(0,0,-G));
+
+	// Place with initial separation as a fraction of the distance moved 
+	// under gravity in one tick when starting from rest.
+	// These leads to penetration depths of: (1.11, 0.83, 0.22, 0.01)
+	float SeparationAlpha = GENERATE(0.0f, 0.25f, 0.75f, 0.99f);
+	float Separation = SeparationAlpha * G * Dt * Dt;
+
+	SECTION("Moving ball with high friction")
+	{
+		const double KinematicPositionTolerance = UE_SMALL_NUMBER;
+		const double DynamicPositionTolerance = UE_KINDA_SMALL_NUMBER;
+
+		// Try different friction values (all should be above static friction threshold)
+		auto Friction = GENERATE(1000.0f, 10.0f, 1.0f);
+
+		// Place a dynamic 50cm radius sphere exactly on the floor (a large kinematic box)
+		FSingleParticlePhysicsProxy* FloorBox = Scene.AddKinematicBox(TEXT("Floor"), FTransform3d(FVector3d(0, 0, -50)), FVector3d(500, 500, 50));
+		FSingleParticlePhysicsProxy* DynamicSphere = Scene.AddDynamicSphere(TEXT("Sphere"), FTransform3d(FVector3d(0, 0, 50 + Separation)), 50);
+
+		// Add horizontal velocity
+		DynamicSphere->GetGameThreadAPI().SetV(FVector3d(100, 0, 0));
+
+		// Set the friction to desired value in a callback
+		// Easier than setting up materials
+		FChaosTestSimCallback* SimCallbacks = Scene.GetSolver()->CreateAndRegisterSimCallbackObject_External<FChaosTestSimCallback>();
+		SimCallbacks->ContactModifier = 
+			[DynamicSphere, Friction](FCollisionContactModifier& Contacts)
+			{
+				for (FContactPairModifier& Contact : Contacts.GetContacts(DynamicSphere->GetHandle_LowLevel()))
+				{
+					Contact.ModifyStaticFriction(Friction);
+					Contact.ModifyDynamicFriction(Friction);
+				}
+			};
+
+		Scene.TickSolver(Dt);
+
+		// Contact point velocity should be zero ish
+		FVector3d V = DynamicSphere->GetGameThreadAPI().GetV();
+		FVector3d W = DynamicSphere->GetGameThreadAPI().GetW();
+		FVector3d VC = V + FVector3d::CrossProduct(W, FVector3d(0,0,-50));
+
+		REQUIRE(Scene.GetSolver()->GetEvolution()->GetCollisionConstraints().NumConstraints() == 1);
+		const FPBDCollisionConstraint& Collision = Scene.GetSolver()->GetEvolution()->GetCollisionConstraints().GetConstConstraint(0);
+		REQUIRE(Collision.NumManifoldPoints() == 1);
+		const FManifoldPoint& ManifoldPoint = Collision.GetManifoldPoint(0);
+		float Phi = ManifoldPoint.ContactPoint.Phi;
+
+		CAPTURE(G, Dt, Friction, Phi, V.X, VC.X);
+
+		CHECK_THAT(VC.Size(), Catch::Matchers::WithinAbs(0.0f, 0.2f));
+	}
+}
diff --git a/Engine/Source/Programs/ChaosPhysicsTests/Private/MassConditioningTests.cpp b/Engine/Source/Programs/ChaosPhysicsTests/Private/MassConditioningTests.cpp
new file mode 100644
--- /dev/null
+++ b/Engine/Source/Programs/ChaosPhysicsTests/Private/MassConditioningTests.cpp
@@ -0,0 +1,67 @@
+// Copyright Epic Games, Inc. All Rights Reserved.
+#include "CoreMinimal.h"
+#include "Chaos/Box.h"
+#include "Chaos/MassConditioning.h"
+#include "Chaos/MassProperties.h"
+#include "ChaosTestHarness.h"
+
+TEST_CASE("ChaosMassConditioningTests", "[Chaos][MassConditioning][unit]")
+{
+	using namespace Chaos;
+
+	FRealSingle MaxDistance = 10.0f;
+	FRealSingle MaxRotationRatio = 2.0f;
+	FRealSingle MaxInvInertiaComponentRatio = 0.0f;
+
+	// Create a box and pretend we have a constraint a long way from the box.
+	// Inertia conditioning will increase the inertia of the box so that impulses
+	// applied at the constraint handle do not introduce crazy rotation.
+	// BUG: Previously this would fail to condition inertia for large objects
+	// because of a tolerance check on the inverse mass and inertia which can 
+	// legitimately be very small.
+	// REBUG: I had to undo the previous fix because it changed behaviour in
+	// a way that can affect existing projects. E.g., vehicles will roll more
+	// easily if InertiaConditioning now affects them but didn't before.
+	SECTION("Small inverse inertia")
+	{
+		FRealSingle BoxExtent = GENERATE(10000.0f, 1000.0f, 100.0f, 10.0f);
+		FRealSingle ConstraintArmMultiplier = GENERATE(1.5f, 5.0f, 10.0f, 100.0f);
+
+		FImplicitBox3 Box(FVec3(-BoxExtent), FVec3(BoxExtent), 0);
+
+		FReal DensityKGPerCM3 = 1.e-3f;
+		FMassProperties MassProperties;
+		CalculateMassPropertiesOfImplicitType(MassProperties, FTransform::Identity, &Box, DensityKGPerCM3);
+		TransformToLocalSpace(MassProperties);
+
+		FRealSingle InvM = 1.0f / MassProperties.Mass;
+		FVec3f InvI = FVec3f(1.0f) / MassProperties.InertiaTensor.GetDiagonal();
+
+		// Pretend we have a constraint with an anchor along the X axis outside the box by some distance.
+		FVec3f ConstraintExtents = FVec3f(ConstraintArmMultiplier * BoxExtent, BoxExtent, BoxExtent);
+
+		// Default Tolerances
+		FInertiaConditioningTolerances Tolerances;
+		CHECK(Tolerances.InvInertiaTolerance == 1.e-4f);	// Tested with this value
+
+		// If the inputs pass the tolerances, Inertia conditioning will increase the Y and Z inertia (i.e., reduce inverse inertia)
+		// so that any rotations applied by the constraint anchor are "small" (based on the settings)
+		FVec3f InertiaConditioning = CalculateInertiaConditioning(InvM, InvI, ConstraintExtents, MaxDistance, MaxRotationRatio, MaxInvInertiaComponentRatio, Tolerances);
+
+		bool bIsSmallInvInertia = InvI.GetMin() < Tolerances.InvInertiaTolerance;
+		if (bIsSmallInvInertia)
+		{
+			CHECK_THAT(InertiaConditioning.X, Catch::Matchers::WithinAbs(1.0f, 0.01f));
+			CHECK_THAT(InertiaConditioning.Y, Catch::Matchers::WithinAbs(1.0f, 0.01f));
+			CHECK_THAT(InertiaConditioning.Z, Catch::Matchers::WithinAbs(1.0f, 0.01f));
+		}
+		else
+		{
+			// TODO: These tests should really verify that applying an impulse at the constaint arm results in a
+			// correction that has at most MaxRotationRatio of its magnitude from rotation
+			CHECK_THAT(InertiaConditioning.X, Catch::Matchers::WithinAbs(1.0f, 0.01f));
+			CHECK(InertiaConditioning.Y < 1.0f);
+			CHECK(InertiaConditioning.Z < 1.0f);
+		}
+	}
+}
\ No newline at end of file
diff --git a/Engine/Source/Programs/ChaosPhysicsTests/Private/PhysicsFilterBuilderTests.cpp b/Engine/Source/Programs/ChaosPhysicsTests/Private/PhysicsFilterBuilderTests.cpp
new file mode 100644
--- /dev/null
+++ b/Engine/Source/Programs/ChaosPhysicsTests/Private/PhysicsFilterBuilderTests.cpp
@@ -0,0 +1,121 @@
+// Copyright Epic Games, Inc. All Rights Reserved.
+#include "CoreMinimal.h"
+
+#include "ChaosTestHarness.h"
+
+#include "Physics/PhysicsFiltering.h"
+
+namespace Chaos::Tests
+{
+	using FShapeFilterData = Chaos::Filter::FShapeFilterData;
+	constexpr uint64 ChannelCount = 64;
+	constexpr uint64 FullChannelMask = std::numeric_limits<uint64>::max();
+
+	TEST_CASE("PhysicsFilterBuilder - SetResponses", "[Chaos][Filter]")
+	{
+		SECTION("Test each block bit")
+		{
+			for (uint64 I = 0; I < ChannelCount; ++I)
+			{
+				FCollisionResponseContainer Container(ECR_Ignore);
+				Container.EnumArray[I] = ECR_Block;
+
+				FPhysicsFilterBuilder Builder;
+				Builder.SetResponses(Container);
+				const FShapeFilterData ShapeFilter = Builder.BuildShapeFilterData();
+
+				CHECK((1LLU << I) == ShapeFilter.GetBlockChannels());
+				CHECK(0 == ShapeFilter.GetOverlapChannels());
+			}
+		}
+		SECTION("Test each overlap bit")
+		{
+			for (uint64 I = 0; I < ChannelCount; ++I)
+			{
+				FCollisionResponseContainer Container(ECR_Ignore);
+				Container.EnumArray[I] = ECR_Overlap;
+
+				FPhysicsFilterBuilder Builder;
+				Builder.SetResponses(Container);
+				const FShapeFilterData ShapeFilter = Builder.BuildShapeFilterData();
+
+				CHECK(0 == ShapeFilter.GetBlockChannels());
+				CHECK((1LLU << I) == ShapeFilter.GetOverlapChannels());
+			}
+		}
+		SECTION("Test full ignore mask")
+		{
+			FCollisionResponseContainer Container(ECR_Ignore);
+
+			FPhysicsFilterBuilder Builder;
+			Builder.SetResponses(Container);
+			const FShapeFilterData ShapeFilter = Builder.BuildShapeFilterData();
+
+			CHECK(0 == ShapeFilter.GetBlockChannels());
+			CHECK(0 == ShapeFilter.GetOverlapChannels());
+		}
+		SECTION("Test full block mask")
+		{
+			FCollisionResponseContainer Container(ECR_Block);
+
+			FPhysicsFilterBuilder Builder;
+			Builder.SetResponses(Container);
+			const Chaos::Filter::FShapeFilterData ShapeFilter = Builder.BuildShapeFilterData();
+
+			CHECK(FullChannelMask == ShapeFilter.GetBlockChannels());
+			CHECK(0 == ShapeFilter.GetOverlapChannels());
+		}
+		SECTION("Test full overlap mask")
+		{
+			FCollisionResponseContainer Container(ECR_Overlap);
+
+			FPhysicsFilterBuilder Builder;
+			Builder.SetResponses(Container);
+			const FShapeFilterData ShapeFilter = Builder.BuildShapeFilterData();
+
+			CHECK(0 == ShapeFilter.GetBlockChannels());
+			CHECK(FullChannelMask == ShapeFilter.GetOverlapChannels());
+		}
+		SECTION("Test mixed masks")
+		{
+			FCollisionResponseContainer Container;
+			for (uint64 I = 0; I < ChannelCount; ++I)
+			{
+				Container.EnumArray[I] = I % 3;
+			}
+
+			FPhysicsFilterBuilder Builder;
+			Builder.SetResponses(Container);
+			const FShapeFilterData ShapeFilter = Builder.BuildShapeFilterData();
+
+			CHECK(0x4924924924924924 == ShapeFilter.GetBlockChannels());
+			CHECK(0x2492492492492492 == ShapeFilter.GetOverlapChannels());
+		}
+	}
+
+	TEST_CASE("PhysicsFilterBuilder - SetResponses - perf", "[Chaos][Filter][!benchmark]")
+	{
+		constexpr uint64 Iterations = 100000;
+
+		FCollisionResponseContainer Container;
+		for (uint64 I = 0; I < ChannelCount; ++I)
+		{
+			Container.EnumArray[I] = I % 3;
+		}
+
+		BENCHMARK("Set Responses")
+		{
+			uint64 BlockResult = 0;
+			uint64 OverlapResult = 0;
+			for (uint64 I = 0; I < Iterations; ++I)
+			{
+				FPhysicsFilterBuilder Builder(nullptr);
+				Builder.SetResponses(Container);
+				const FShapeFilterData ShapeFilter = Builder.BuildShapeFilterData();
+				BlockResult |= ShapeFilter.GetBlockChannels();
+				OverlapResult |= ShapeFilter.GetOverlapChannels();
+			}
+			return BlockResult | OverlapResult;
+		};
+	}
+} // namespace Chaos::Tests
diff --git a/Engine/Source/Programs/ChaosPhysicsTests/Private/RestitutionTests.cpp b/Engine/Source/Programs/ChaosPhysicsTests/Private/RestitutionTests.cpp
new file mode 100644
--- /dev/null
+++ b/Engine/Source/Programs/ChaosPhysicsTests/Private/RestitutionTests.cpp
@@ -0,0 +1,86 @@
+// Copyright Epic Games, Inc. All Rights Reserved.
+#include "CoreMinimal.h"
+#include "Chaos/Box.h"
+#include "Chaos/Collision/PBDCollisionConstraint.h"
+#include "Chaos/ContactModification.h"
+#include "Chaos/MassProperties.h"
+#include "Chaos/PBDCollisionConstraints.h"
+#include "Chaos/PBDRigidsEvolution.h"
+#include "Chaos/SimCallbackObject.h"
+#include "ChaosSolversModule.h"
+#include "ChaosTestHarness.h"
+#include "ChaosTestScene.h"
+#include "EventsData.h"
+#include "PBDRigidsSolver.h"
+#include "PhysicsProxy/SingleParticlePhysicsProxy.h"
+#include "ScopedCVar.h"
+
+//
+// Test that a dropped box bounces correctly
+//
+TEST_CASE("ChaosRestitutionTests", "[Chaos][Restitution][unit]")
+{
+	using namespace Chaos;
+	using namespace Chaos::LowLevelTest;
+
+	const float Dt = 1.0f / 60.0f;
+
+	FChaosTestScene Scene;
+	Scene.SetGravity(FVector3d(0));
+
+	// 4 iterations required for the thresholds we set
+	int32 NumVelIts = GENERATE(8, 6, 4);
+	Scene.GetSolver()->SetVelocityIterations(NumVelIts);
+
+	// Try different friction values
+	float Restitution = 1.0f;
+	float Friction = GENERATE(0.0f, 1.0f);
+
+	// Place just above the ground with high velocity
+	float DropHeight = 1.0f;
+	float InitialVelZ = -400.0f;
+
+	SECTION("A box dropped onto flat ground should bounce straight back")
+	{
+		// Place a dynamic box above the floor
+		FSingleParticlePhysicsProxy* FloorBox = Scene.AddKinematicBox(TEXT("Floor"), FTransform3d(FVector3d(0, 0, -50)), FVector3d(500, 500, 50));
+		FSingleParticlePhysicsProxy* DynamicSphere = Scene.AddDynamicBox(TEXT("Box"), FTransform3d(FVector3d(0, 0, 50 + DropHeight)), FVector3d(50, 50, 50));
+
+		DynamicSphere->GetGameThreadAPI().SetV(FVector3d(0, 0, InitialVelZ));
+
+		// Set the friction and restitution to desired value in a callback
+		// Easier than setting up materials
+		FChaosTestSimCallback* SimCallbacks = Scene.GetSolver()->CreateAndRegisterSimCallbackObject_External<FChaosTestSimCallback>();
+		SimCallbacks->ContactModifier = 
+			[DynamicSphere, Friction, Restitution](FCollisionContactModifier& Contacts)
+			{
+				for (FContactPairModifier& Contact : Contacts.GetContacts(DynamicSphere->GetHandle_LowLevel()))
+				{
+					Contact.ModifyStaticFriction(Friction);
+					Contact.ModifyDynamicFriction(Friction);
+					Contact.ModifyRestitution(Restitution);
+				}
+			};
+
+		FVector3d PreV = DynamicSphere->GetGameThreadAPI().GetV();
+		FVector3d PreW = DynamicSphere->GetGameThreadAPI().GetW();
+
+		Scene.TickSolver(Dt);
+
+		FVector3d PostV = DynamicSphere->GetGameThreadAPI().GetV();
+		FVector3d PostW = DynamicSphere->GetGameThreadAPI().GetW();
+
+		const FPBDCollisionConstraints& CollisionConstraints = Scene.GetSolver()->GetEvolution()->GetCollisionConstraints();
+		REQUIRE(CollisionConstraints.NumConstraints() == 1);
+
+		const FPBDCollisionConstraint& Collision =CollisionConstraints.GetConstConstraint(0);
+		REQUIRE(Collision.NumManifoldPoints() == 4);
+
+		const FManifoldPoint& ManifoldPoint = Collision.GetManifoldPoint(0);
+		float Phi = ManifoldPoint.ContactPoint.Phi;
+
+		CAPTURE(NumVelIts, Friction, Restitution, InitialVelZ, Phi, PreV.Z, PostV.Z, PreW.Size(), PostW.Size());
+		CHECK_THAT(PostV.Z, Catch::Matchers::WithinAbs(-Restitution * PreV.Z, 1.0f));
+		CHECK_THAT(PostW.Size(), Catch::Matchers::WithinAbs(0.0f, 0.01f));
+	}
+}
diff --git a/Engine/Source/Programs/ChaosPhysicsTests/Private/ScopedCVar.h b/Engine/Source/Programs/ChaosPhysicsTests/Private/ScopedCVar.h
new file mode 100644
--- /dev/null
+++ b/Engine/Source/Programs/ChaosPhysicsTests/Private/ScopedCVar.h
@@ -0,0 +1,93 @@
+// Copyright Epic Games, Inc. All Rights Reserved.
+#pragma once
+
+#include "HAL/IConsoleManager.h"
+
+template<typename T>
+class FCVarSetter
+{
+public:
+	FCVarSetter(const FString& InConsoleVariableName)
+		: ConsoleVariableName(InConsoleVariableName)
+		, OriginalValue()
+		, bIsSet(false)
+	{
+	}
+
+	virtual ~FCVarSetter()
+	{
+	}
+
+	void Set(T Value);
+
+	void Restore()
+	{
+		if (bIsSet)
+		{
+			if (IConsoleVariable* ConsoleVariable = IConsoleManager::Get().FindConsoleVariable(*ConsoleVariableName))
+			{
+				ConsoleVariable->AsVariable()->SetWithCurrentPriority(OriginalValue);
+			}
+			bIsSet = false;
+		}
+	}
+
+private:
+	FString ConsoleVariableName;
+	T OriginalValue;
+	bool bIsSet;
+};
+
+template<>
+FORCEINLINE void FCVarSetter<bool>::Set(bool Value)
+{
+	check(!bIsSet);
+	if (IConsoleVariable* ConsoleVariable = IConsoleManager::Get().FindConsoleVariable(*ConsoleVariableName))
+	{
+		OriginalValue = ConsoleVariable->GetBool();
+		ConsoleVariable->AsVariable()->SetWithCurrentPriority(Value);
+		bIsSet = true;
+	}
+}
+
+template<>
+FORCEINLINE void FCVarSetter<int32>::Set(int32 Value)
+{
+	check(!bIsSet);
+	if (IConsoleVariable* ConsoleVariable = IConsoleManager::Get().FindConsoleVariable(*ConsoleVariableName))
+	{
+		OriginalValue = ConsoleVariable->GetInt();
+		ConsoleVariable->AsVariable()->SetWithCurrentPriority(Value);
+		bIsSet = true;
+	}
+}
+
+template<>
+FORCEINLINE void FCVarSetter<float>::Set(float Value)
+{
+	check(!bIsSet);
+	if (IConsoleVariable* ConsoleVariable = IConsoleManager::Get().FindConsoleVariable(*ConsoleVariableName))
+	{
+		OriginalValue = ConsoleVariable->GetFloat();
+		ConsoleVariable->AsVariable()->SetWithCurrentPriority(Value);
+		bIsSet = true;
+	}
+}
+
+template<typename T>
+class FCVarScopeGuard
+{
+public:
+	FCVarScopeGuard(const FString& InConsoleVariableName, T InValue) 
+		: CVarSetter(InConsoleVariableName)
+	{
+		CVarSetter.Set(InValue);
+	}
+
+	~FCVarScopeGuard()
+	{
+		CVarSetter.Restore();
+	}
+private:
+	FCVarSetter<T> CVarSetter;
+};
diff --git a/Engine/Source/Programs/ChaosPhysicsTests/Private/SpatialAccelerationBroadPhaseTests.cpp b/Engine/Source/Programs/ChaosPhysicsTests/Private/SpatialAccelerationBroadPhaseTests.cpp
new file mode 100644
--- /dev/null
+++ b/Engine/Source/Programs/ChaosPhysicsTests/Private/SpatialAccelerationBroadPhaseTests.cpp
@@ -0,0 +1,126 @@
+// Copyright Epic Games, Inc. All Rights Reserved.
+
+#include "CoreMinimal.h"
+#include "ChaosTestHarness.h"
+#include "ChaosTestScene.h"
+#include "ChaosSolversModule.h"
+
+#include "Chaos/Box.h"
+#include "PBDRigidsSolver.h"
+#include "PhysicsProxy/SingleParticlePhysicsProxy.h"
+
+namespace Chaos
+{
+	TEST_CASE("SpatialAccelerationBroadPhase - PrePreFilter", "[Chaos][BroadPhase][unit]")
+	{
+		using namespace Chaos::LowLevelTest;
+
+		// This test is for a bug where the full collision filter data was or'd together, including word 3's collision channel index, which would do invalid filtering with union shapes.
+		constexpr float Dt = 1.0f / 30.0f;
+		FChaosTestScene Scene;
+
+		const FVector3d HalfExtents(50, 50, 50);
+		TArray<FImplicitObjectPtr> Geometry
+		{
+			FImplicitObjectPtr(new Chaos::FImplicitBox3(-HalfExtents, HalfExtents)),
+			FImplicitObjectPtr(new Chaos::FImplicitBox3(-HalfExtents, HalfExtents)),
+		};
+		FImplicitObjectPtr UnionImplicit = FImplicitObjectPtr(new FImplicitObjectUnion(MoveTemp(Geometry)));
+		FSingleParticlePhysicsProxy* UnionProxy0 = Scene.AddKinematicProxy(TEXT("Union0"), FTransform3d(FVector3d(0, 0, 0)), UnionImplicit);
+		FSingleParticlePhysicsProxy* UnionProxy1 = Scene.AddDynamicProxy(TEXT("Union1"), FTransform3d(FVector3d(0, 0, 60)), UnionImplicit);
+
+		Private::FCollisionConstraintAllocator CollisionAllocator;
+		CollisionAllocator.SetMaxContexts(1);
+		FCollisionDetectorSettings Settings;
+		Chaos::FSpatialAccelerationBroadPhase& BroadPhase = Scene.GetSolver()->GetEvolution()->GetBroadPhase();
+
+		auto SetShapeFilter = [](FSingleParticlePhysicsProxy* Proxy, const int32 ShapeIndex, uint8 CollisionChannelIndex, Chaos::EFilterFlags Flags, uint32 BlockMask)
+		{
+			Filter::FShapeFilterBuilder Builder;
+			Builder.SetCollisionChannelIndex(CollisionChannelIndex);
+			Builder.SetBlockChannelMask(BlockMask);
+			Builder.SetFilterFlags(Flags);
+			const Filter::FShapeFilterData FilterData = Builder.Build();
+
+			Proxy->GetGameThreadAPI().SetShapeFilterData(ShapeIndex, FilterData);
+		};
+
+		// Test cases that should produce no overlap pairs.
+		{
+			// Test everything blocking channel 3. Proxy0 is set to channels 1 and 2 which if incorrectly 
+			// bitwise or'd together would turn into channel 3. This should not return any pairs.
+			SetShapeFilter(UnionProxy0, 0, 1, EFilterFlags::All, 0b01000);
+			SetShapeFilter(UnionProxy0, 1, 2, EFilterFlags::All, 0b01000);
+			SetShapeFilter(UnionProxy1, 0, 3, EFilterFlags::All, 0b01000);
+			SetShapeFilter(UnionProxy1, 1, 3, EFilterFlags::All, 0b01000);
+
+			Scene.TickSolver(Dt);
+			CHECK(BroadPhase.GetNumBroadPhasePairs() == 0);
+
+			// Test the opposite order to ensure ordering doesn't matter.
+			SetShapeFilter(UnionProxy0, 0, 3, EFilterFlags::All, 0b01000);
+			SetShapeFilter(UnionProxy0, 1, 3, EFilterFlags::All, 0b01000);
+			SetShapeFilter(UnionProxy1, 0, 1, EFilterFlags::All, 0b01000);
+			SetShapeFilter(UnionProxy1, 1, 2, EFilterFlags::All, 0b01000);
+
+			Scene.TickSolver(Dt);
+			CHECK(BroadPhase.GetNumBroadPhasePairs() == 0);
+		}
+
+		// Now test cases that should produce an overlap.
+		{
+			// Test where Shape10 should overlap Shape01.
+			SetShapeFilter(UnionProxy0, 0, 1, EFilterFlags::All, 0b01000);
+			SetShapeFilter(UnionProxy0, 1, 2, EFilterFlags::All, 0b01000);
+			SetShapeFilter(UnionProxy1, 0, 3, EFilterFlags::All, 0b00100);
+			SetShapeFilter(UnionProxy1, 1, 3, EFilterFlags::All, 0b01000);
+
+			// Test where Shape10 should overlap Shape00.
+			Scene.TickSolver(Dt);
+			CHECK(BroadPhase.GetNumBroadPhasePairs() == 1);
+
+			SetShapeFilter(UnionProxy0, 0, 1, EFilterFlags::All, 0b01000);
+			SetShapeFilter(UnionProxy0, 1, 2, EFilterFlags::All, 0b01000);
+			SetShapeFilter(UnionProxy1, 0, 3, EFilterFlags::All, 0b00010);
+			SetShapeFilter(UnionProxy1, 1, 3, EFilterFlags::All, 0b01000);
+			Scene.TickSolver(Dt);
+			CHECK(BroadPhase.GetNumBroadPhasePairs() == 1);
+
+			// Test where Shape11 should overlap both shapes on proxy0.
+			SetShapeFilter(UnionProxy0, 0, 1, EFilterFlags::All, 0b01000);
+			SetShapeFilter(UnionProxy0, 1, 2, EFilterFlags::All, 0b01000);
+			SetShapeFilter(UnionProxy1, 0, 3, EFilterFlags::All, 0b01000);
+			SetShapeFilter(UnionProxy1, 1, 3, EFilterFlags::All, 0b00110);
+			Scene.TickSolver(Dt);
+			CHECK(BroadPhase.GetNumBroadPhasePairs() == 1);
+		}
+
+		// Test the opposite ordering to be safe.
+		{
+			// Test where Shape00 should overlap Shape11.
+			SetShapeFilter(UnionProxy1, 0, 1, EFilterFlags::All, 0b01000);
+			SetShapeFilter(UnionProxy1, 1, 2, EFilterFlags::All, 0b01000);
+			SetShapeFilter(UnionProxy0, 0, 3, EFilterFlags::All, 0b00100);
+			SetShapeFilter(UnionProxy0, 1, 3, EFilterFlags::All, 0b01000);
+
+			// Test where Shape00 should overlap Shape10
+			Scene.TickSolver(Dt);
+			CHECK(BroadPhase.GetNumBroadPhasePairs() == 1);
+
+			SetShapeFilter(UnionProxy1, 0, 1, EFilterFlags::All, 0b01000);
+			SetShapeFilter(UnionProxy1, 1, 2, EFilterFlags::All, 0b01000);
+			SetShapeFilter(UnionProxy0, 0, 3, EFilterFlags::All, 0b00010);
+			SetShapeFilter(UnionProxy0, 1, 3, EFilterFlags::All, 0b01000);
+			Scene.TickSolver(Dt);
+			CHECK(BroadPhase.GetNumBroadPhasePairs() == 1);
+
+			// Test where Shape01 should overlap both shapes on proxy1
+			SetShapeFilter(UnionProxy1, 0, 1, EFilterFlags::All, 0b01000);
+			SetShapeFilter(UnionProxy1, 1, 2, EFilterFlags::All, 0b01000);
+			SetShapeFilter(UnionProxy0, 0, 3, EFilterFlags::All, 0b01000);
+			SetShapeFilter(UnionProxy0, 1, 3, EFilterFlags::All, 0b00110);
+			Scene.TickSolver(Dt);
+			CHECK(BroadPhase.GetNumBroadPhasePairs() == 1);
+		}
+	}
+} // namespace Chaos
diff --git a/Engine/Source/Programs/ChaosPhysicsTests/Private/SweepTests.cpp b/Engine/Source/Programs/ChaosPhysicsTests/Private/SweepTests.cpp
new file mode 100644
--- /dev/null
+++ b/Engine/Source/Programs/ChaosPhysicsTests/Private/SweepTests.cpp
@@ -0,0 +1,1805 @@
+// Copyright Epic Games, Inc. All Rights Reserved.
+
+#include "CoreMinimal.h"
+#include "ChaosTestHarness.h"
+#include <catch2/matchers/catch_matchers_all.hpp>
+
+#include "Chaos/Box.h"
+#include "Chaos/Sphere.h"
+#include "Chaos/Capsule.h"
+#include "Chaos/GeometryQueries.h"
+#include "Chaos/Sweeps.h"
+
+namespace Chaos
+{
+	TEST_CASE("SweepSphereVsSphere", "[Chaos][Sweeps][unit]")
+	{
+		const bool bComputeMTD = true;
+		const FReal Thickness = 0;
+		const FReal Length = 100;
+		const TSphere<double, 3> TestGeom(FVec3::ZeroVector, 5);
+		const TSphere<double, 3> SweepGeom(FVec3::ZeroVector, 2);
+
+		FReal ResultTime;
+		FVec3 ResultPosition;
+		FVec3 ResultNormal;
+		FVec3 ResultFaceNormal = FVec3::ZeroVector;
+		int32 ResultFaceIndex = INDEX_NONE;
+
+		SECTION("Sweep From -X")
+		{
+			const FRigidTransform3 TestGeomTransform(FVec3(1, 2, 3), FRotation3::FromIdentity());
+			const FRigidTransform3 SweepGeomTransform(FVec3(-9, 2, 3), FRotation3::FromIdentity());
+			const FVec3 Dir(1, 0, 0);
+
+			const bool bResult = SweepQuery(TestGeom, TestGeomTransform, SweepGeom, SweepGeomTransform, Dir, Length, ResultTime, ResultPosition, ResultNormal, ResultFaceIndex, ResultFaceNormal, Thickness, bComputeMTD);
+			CHECK(bResult);
+			CHECK_THAT(ResultTime, Catch::Matchers::WithinRel(3.0, UE_DOUBLE_KINDA_SMALL_NUMBER));
+			CHECK_THAT(ResultPosition, Catch::ApproxEq(FVec3(-4, 2, 3)));
+			CHECK_THAT(ResultNormal, Catch::ApproxEq(FVec3(-1, 0, 0)));
+			CHECK_THAT(ResultFaceNormal, Catch::ApproxEq(FVec3(-1, 0, 0)));
+			CHECK(ResultFaceIndex == INDEX_NONE);
+		}
+		SECTION("Sweep From +X")
+		{
+			const FRigidTransform3 TestGeomTransform(FVec3(1, 2, 3), FRotation3::FromIdentity());
+			const FRigidTransform3 SweepGeomTransform(FVec3(11, 2, 3), FRotation3::FromIdentity());
+			const FVec3 Dir(-1, 0, 0);
+
+			const bool bResult = SweepQuery(TestGeom, TestGeomTransform, SweepGeom, SweepGeomTransform, Dir, Length, ResultTime, ResultPosition, ResultNormal, ResultFaceIndex, ResultFaceNormal, Thickness, bComputeMTD);
+			CHECK(bResult);
+			CHECK_THAT(ResultTime, Catch::Matchers::WithinRel(3.0, UE_DOUBLE_KINDA_SMALL_NUMBER));
+			CHECK_THAT(ResultPosition, Catch::ApproxEq(FVec3(6, 2, 3)));
+			CHECK_THAT(ResultNormal, Catch::ApproxEq(FVec3(1, 0, 0)));
+			CHECK_THAT(ResultFaceNormal, Catch::ApproxEq(FVec3(1, 0, 0)));
+			CHECK(ResultFaceIndex == INDEX_NONE);
+		}
+		SECTION("Sweep From -Y")
+		{
+			const FRigidTransform3 TestGeomTransform(FVec3(1, 2, 3), FRotation3::FromIdentity());
+			const FRigidTransform3 SweepGeomTransform(FVec3(1, -8, 3), FRotation3::FromIdentity());
+			const FVec3 Dir(0, 1, 0);
+
+			const bool bResult = SweepQuery(TestGeom, TestGeomTransform, SweepGeom, SweepGeomTransform, Dir, Length, ResultTime, ResultPosition, ResultNormal, ResultFaceIndex, ResultFaceNormal, Thickness, bComputeMTD);
+			CHECK(bResult);
+			CHECK_THAT(ResultTime, Catch::Matchers::WithinRel(3.0, UE_DOUBLE_KINDA_SMALL_NUMBER));
+			CHECK_THAT(ResultPosition, Catch::ApproxEq(FVec3(1, -3, 3)));
+			CHECK_THAT(ResultNormal, Catch::ApproxEq(FVec3(0, -1, 0)));
+			CHECK_THAT(ResultFaceNormal, Catch::ApproxEq(FVec3(0, -1, 0)));
+			CHECK(ResultFaceIndex == INDEX_NONE);
+		}
+		SECTION("Sweep From +Y")
+		{
+			const FRigidTransform3 TestGeomTransform(FVec3(1, 2, 3), FRotation3::FromIdentity());
+			const FRigidTransform3 SweepGeomTransform(FVec3(1, 12, 3), FRotation3::FromIdentity());
+			const FVec3 Dir(0, -1, 0);
+
+			const bool bResult = SweepQuery(TestGeom, TestGeomTransform, SweepGeom, SweepGeomTransform, Dir, Length, ResultTime, ResultPosition, ResultNormal, ResultFaceIndex, ResultFaceNormal, Thickness, bComputeMTD);
+			CHECK(bResult);
+			CHECK_THAT(ResultTime, Catch::Matchers::WithinRel(3.0, UE_DOUBLE_KINDA_SMALL_NUMBER));
+			CHECK_THAT(ResultPosition, Catch::ApproxEq(FVec3(1, 7, 3)));
+			CHECK_THAT(ResultNormal, Catch::ApproxEq(FVec3(0, 1, 0)));
+			CHECK_THAT(ResultFaceNormal, Catch::ApproxEq(FVec3(0, 1, 0)));
+			CHECK(ResultFaceIndex == INDEX_NONE);
+		}
+		SECTION("Sweep From -Z")
+		{
+			const FRigidTransform3 TestGeomTransform(FVec3(1, 2, 3), FRotation3::FromIdentity());
+			const FRigidTransform3 SweepGeomTransform(FVec3(1, 2, -7), FRotation3::FromIdentity());
+			const FVec3 Dir(0, 0, 1);
+
+			const bool bResult = SweepQuery(TestGeom, TestGeomTransform, SweepGeom, SweepGeomTransform, Dir, Length, ResultTime, ResultPosition, ResultNormal, ResultFaceIndex, ResultFaceNormal, Thickness, bComputeMTD);
+			CHECK(bResult);
+			CHECK_THAT(ResultTime, Catch::Matchers::WithinRel(3.0, UE_DOUBLE_KINDA_SMALL_NUMBER));
+			CHECK_THAT(ResultPosition, Catch::ApproxEq(FVec3(1, 2, -2)));
+			CHECK_THAT(ResultNormal, Catch::ApproxEq(FVec3(0, 0, -1)));
+			CHECK_THAT(ResultFaceNormal, Catch::ApproxEq(FVec3(0, 0, -1)));
+			CHECK(ResultFaceIndex == INDEX_NONE);
+		}
+		SECTION("Sweep From +Z")
+		{
+			const FRigidTransform3 TestGeomTransform(FVec3(1, 2, 3), FRotation3::FromIdentity());
+			const FRigidTransform3 SweepGeomTransform(FVec3(1, 2, 13), FRotation3::FromIdentity());
+			const FVec3 Dir(0, 0, -1);
+
+			const bool bResult = SweepQuery(TestGeom, TestGeomTransform, SweepGeom, SweepGeomTransform, Dir, Length, ResultTime, ResultPosition, ResultNormal, ResultFaceIndex, ResultFaceNormal, Thickness, bComputeMTD);
+			CHECK(bResult);
+			CHECK_THAT(ResultTime, Catch::Matchers::WithinRel(3.0, UE_DOUBLE_KINDA_SMALL_NUMBER));
+			CHECK_THAT(ResultPosition, Catch::ApproxEq(FVec3(1, 2, 8)));
+			CHECK_THAT(ResultNormal, Catch::ApproxEq(FVec3(0, 0, 1)));
+			CHECK_THAT(ResultFaceNormal, Catch::ApproxEq(FVec3(0, 0, 1)));
+			CHECK(ResultFaceIndex == INDEX_NONE);
+		}
+		SECTION("Sweep From -X with diagonal hit")
+		{
+			// Use law of sines to figure out the positional offset for a 45 degree angle hit
+			const FReal ZOffset = (TestGeom.GetRadiusf() + SweepGeom.GetRadiusf()) * FMath::Sin(FMath::DegreesToRadians(45)) / FMath::Sin(FMath::DegreesToRadians(90));
+			const FRigidTransform3 TestGeomTransform(FVec3(1, 2, 3), FRotation3::FromIdentity());
+			const FRigidTransform3 SweepGeomTransform(FVec3(-9, 2, 3 + ZOffset), FRotation3::FromIdentity());
+			const FVec3 Dir(1, 0, 0);
+			const FVec3 ExpectedNormal = FVec3(-1, 0, 1).GetSafeNormal();
+			const FVec3 ExpectedPosition = TestGeomTransform.GetLocation() + ExpectedNormal * TestGeom.GetRadiusf();
+
+			const bool bResult = SweepQuery(TestGeom, TestGeomTransform, SweepGeom, SweepGeomTransform, Dir, Length, ResultTime, ResultPosition, ResultNormal, ResultFaceIndex, ResultFaceNormal, Thickness, bComputeMTD);
+			CHECK(bResult);
+			CHECK_THAT(ResultTime, Catch::Matchers::WithinRel(5.05025, UE_DOUBLE_KINDA_SMALL_NUMBER));
+			CHECK_THAT(ResultPosition, Catch::ApproxEq(ExpectedPosition));
+			CHECK_THAT(ResultNormal, Catch::ApproxEq(ExpectedNormal));
+			CHECK_THAT(ResultFaceNormal, Catch::ApproxEq(ExpectedNormal));
+			CHECK(ResultFaceIndex == INDEX_NONE);
+		}
+		SECTION("Sweep: Past max distance")
+		{
+			const FRigidTransform3 TestGeomTransform(FVec3(1, 2, 3), FRotation3::FromIdentity());
+			const FRigidTransform3 SweepGeomTransform(FVec3(1.05 + Length + 7, 2, 3), FRotation3::FromIdentity());
+			const FVec3 Dir(-1, 0, 0);
+
+			const bool bResult = SweepQuery(TestGeom, TestGeomTransform, SweepGeom, SweepGeomTransform, Dir, Length, ResultTime, ResultPosition, ResultNormal, ResultFaceIndex, ResultFaceNormal, Thickness, bComputeMTD);
+			CHECK(!bResult);
+		}
+		SECTION("Sweep: Miss")
+		{
+			const FRigidTransform3 TestGeomTransform(FVec3(1, 2, 3), FRotation3::FromIdentity());
+			const FRigidTransform3 SweepGeomTransform(FVec3(-9, 10, 3), FRotation3::FromIdentity());
+			const FVec3 Dir(1, 0, 0);
+
+			const bool bResult = SweepQuery(TestGeom, TestGeomTransform, SweepGeom, SweepGeomTransform, Dir, Length, ResultTime, ResultPosition, ResultNormal, ResultFaceIndex, ResultFaceNormal, Thickness, bComputeMTD);
+			CHECK(!bResult);
+		}
+		SECTION("Shallow Initial Overlap")
+		{
+			const FRigidTransform3 TestGeomTransform(FVec3(1, 2, 3), FRotation3::FromIdentity());
+			const FRigidTransform3 SweepGeomTransform(FVec3(-5, 2, 3), FRotation3::FromIdentity());
+			const FVec3 Dir(1, 0, 0);
+
+			const bool bResult = SweepQuery(TestGeom, TestGeomTransform, SweepGeom, SweepGeomTransform, Dir, Length, ResultTime, ResultPosition, ResultNormal, ResultFaceIndex, ResultFaceNormal, Thickness, bComputeMTD);
+			CHECK(bResult);
+			CHECK_THAT(ResultTime, Catch::Matchers::WithinRel(-1.0, UE_DOUBLE_KINDA_SMALL_NUMBER));
+			CHECK_THAT(ResultPosition, Catch::ApproxEq(FVec3(-4, 2, 3)));
+			CHECK_THAT(ResultNormal, Catch::ApproxEq(FVec3(-1, 0, 0)));
+			CHECK_THAT(ResultFaceNormal, Catch::ApproxEq(FVec3(-1, 0, 0)));
+			CHECK(ResultFaceIndex == INDEX_NONE);
+		}
+		SECTION("Deep Initial Overlap")
+		{
+			const FRigidTransform3 TestGeomTransform(FVec3(1, 2, 3), FRotation3::FromIdentity());
+			const FRigidTransform3 SweepGeomTransform(FVec3(0, 2, 3), FRotation3::FromIdentity());
+			const FVec3 Dir(1, 0, 0);
+
+			const bool bResult = SweepQuery(TestGeom, TestGeomTransform, SweepGeom, SweepGeomTransform, Dir, Length, ResultTime, ResultPosition, ResultNormal, ResultFaceIndex, ResultFaceNormal, Thickness, bComputeMTD);
+			CHECK(bResult);
+			CHECK_THAT(ResultTime, Catch::Matchers::WithinRel(-6.0, UE_DOUBLE_KINDA_SMALL_NUMBER));
+			CHECK_THAT(ResultPosition, Catch::ApproxEq(FVec3(-4, 2, 3)));
+			CHECK_THAT(ResultNormal, Catch::ApproxEq(FVec3(-1, 0, 0)));
+			CHECK_THAT(ResultFaceNormal, Catch::ApproxEq(FVec3(-1, 0, 0)));
+			CHECK(ResultFaceIndex == INDEX_NONE);
+		}
+		SECTION("Initial Overlap Behind")
+		{
+			const FRigidTransform3 TestGeomTransform(FVec3(1, 2, 3), FRotation3::FromIdentity());
+			const FRigidTransform3 SweepGeomTransform(FVec3(2, 2, 3), FRotation3::FromIdentity());
+			const FVec3 Dir(1, 0, 0);
+
+			const bool bResult = SweepQuery(TestGeom, TestGeomTransform, SweepGeom, SweepGeomTransform, Dir, Length, ResultTime, ResultPosition, ResultNormal, ResultFaceIndex, ResultFaceNormal, Thickness, bComputeMTD);
+			CHECK(bResult);
+			CHECK_THAT(ResultTime, Catch::Matchers::WithinRel(-6, UE_DOUBLE_KINDA_SMALL_NUMBER));
+			CHECK_THAT(ResultPosition, Catch::ApproxEq(FVec3(6, 2, 3)));
+			CHECK_THAT(ResultNormal, Catch::ApproxEq(FVec3(1, 0, 0)));
+			CHECK_THAT(ResultFaceNormal, Catch::ApproxEq(FVec3(1, 0, 0)));
+			CHECK(ResultFaceIndex == INDEX_NONE);
+		}
+		SECTION("With Thickness")
+		{
+			const FRigidTransform3 TestGeomTransform(FVec3(1, 2, 3), FRotation3::FromIdentity());
+			const FRigidTransform3 SweepGeomTransform(FVec3(-9, 2, 3), FRotation3::FromIdentity());
+			const FVec3 Dir(1, 0, 0);
+			const FReal TestThickness = 1;
+
+			const bool bResult = SweepQuery(TestGeom, TestGeomTransform, SweepGeom, SweepGeomTransform, Dir, Length, ResultTime, ResultPosition, ResultNormal, ResultFaceIndex, ResultFaceNormal, TestThickness, bComputeMTD);
+			CHECK(bResult);
+			CHECK_THAT(ResultTime, Catch::Matchers::WithinRel(2.0, UE_DOUBLE_KINDA_SMALL_NUMBER));
+			CHECK_THAT(ResultPosition, Catch::ApproxEq(FVec3(-4, 2, 3)));
+			CHECK_THAT(ResultNormal, Catch::ApproxEq(FVec3(-1, 0, 0)));
+			CHECK_THAT(ResultFaceNormal, Catch::ApproxEq(FVec3(-1, 0, 0)));
+			CHECK(ResultFaceIndex == INDEX_NONE);
+		}
+		SECTION("Zero Length Sweep: Miss")
+		{
+			const FRigidTransform3 TestGeomTransform(FVec3(1, 2, 3), FRotation3::FromIdentity());
+			const FRigidTransform3 SweepGeomTransform(FVec3(-9, 2, 3), FRotation3::FromIdentity());
+			const FVec3 Dir(1, 0, 0);
+			const FReal TestLength = 0;
+
+			const bool bResult = SweepQuery(TestGeom, TestGeomTransform, SweepGeom, SweepGeomTransform, Dir, TestLength, ResultTime, ResultPosition, ResultNormal, ResultFaceIndex, ResultFaceNormal, Thickness, bComputeMTD);
+			CHECK(bResult == false);
+		}
+		SECTION("Zero Length Sweep: Hit")
+		{
+			const FRigidTransform3 TestGeomTransform(FVec3(1, 2, 3), FRotation3::FromIdentity());
+			const FRigidTransform3 SweepGeomTransform(FVec3(-5, 2, 3), FRotation3::FromIdentity());
+			const FVec3 Dir(1, 0, 0);
+			const FReal TestLength = 0;
+
+			const bool bResult = SweepQuery(TestGeom, TestGeomTransform, SweepGeom, SweepGeomTransform, Dir, TestLength, ResultTime, ResultPosition, ResultNormal, ResultFaceIndex, ResultFaceNormal, Thickness, bComputeMTD);
+			CHECK(bResult);
+			CHECK_THAT(ResultTime, Catch::Matchers::WithinRel(-1.0, UE_DOUBLE_KINDA_SMALL_NUMBER));
+			CHECK_THAT(ResultPosition, Catch::ApproxEq(FVec3(-4, 2, 3)));
+			CHECK_THAT(ResultNormal, Catch::ApproxEq(FVec3(-1, 0, 0)));
+			CHECK_THAT(ResultFaceNormal, Catch::ApproxEq(FVec3(-1, 0, 0)));
+			CHECK(ResultFaceIndex == INDEX_NONE);
+		}
+	}
+
+	TEST_CASE("SweepSphereVsSphere - Perf", "[Chaos][Sweeps][!benchmark]")
+	{
+		const bool bComputeMTD = true;
+		const FReal Thickness = 0;
+		const FVec3 Dir(1, 0, 0);
+		const FReal Length = 100;
+		const TSphere<double, 3> TestGeom(FVec3::ZeroVector, 5);
+		const TSphere<double, 3> SweepGeom(FVec3::ZeroVector, 2);
+		const FRigidTransform3 TestGeomTransform(FVec3(1, 2, 3), FRotation3::FromIdentity());
+
+		FReal ResultTime;
+		FVec3 ResultPosition, ResultNormal, ResultFaceNormal;
+		int32 ResultFaceIndex = INDEX_NONE;
+
+		BENCHMARK("Basic Hit")
+		{
+			const FRigidTransform3 SweepGeomTransform(FVec3(-9, 2, 3), FRotation3::FromIdentity());
+			bool bResult = false;
+			for (int32 I = 0; I < 1000; ++I)
+			{
+				bResult |= SweepQuery(TestGeom, TestGeomTransform, SweepGeom, SweepGeomTransform, Dir, Length, ResultTime, ResultPosition, ResultNormal, ResultFaceIndex, ResultFaceNormal, Thickness, bComputeMTD);
+			}
+			return bResult;
+		};
+		BENCHMARK("Miss")
+		{
+			const FRigidTransform3 SweepGeomTransform(FVec3(-9, 10, 3), FRotation3::FromIdentity());
+			bool bResult = false;
+			for (int32 I = 0; I < 1000; ++I)
+			{
+				bResult |= SweepQuery(TestGeom, TestGeomTransform, SweepGeom, SweepGeomTransform, Dir, Length, ResultTime, ResultPosition, ResultNormal, ResultFaceIndex, ResultFaceNormal, Thickness, bComputeMTD);
+			}
+			return bResult;
+		};
+		BENCHMARK("Shallow Initial Overlap")
+		{
+			const FRigidTransform3 SweepGeomTransform(FVec3(-5, 2, 3), FRotation3::FromIdentity());
+			bool bResult = false;
+			for (int32 I = 0; I < 1000; ++I)
+			{
+				bResult |= SweepQuery(TestGeom, TestGeomTransform, SweepGeom, SweepGeomTransform, Dir, Length, ResultTime, ResultPosition, ResultNormal, ResultFaceIndex, ResultFaceNormal, Thickness, bComputeMTD);
+			}
+			return bResult;
+		};
+		BENCHMARK("Deep Initial Overlap")
+		{
+			const FRigidTransform3 SweepGeomTransform(FVec3(0, 2, 3), FRotation3::FromIdentity());
+			bool bResult = false;
+			for (int32 I = 0; I < 1000; ++I)
+			{
+				bResult |= SweepQuery(TestGeom, TestGeomTransform, SweepGeom, SweepGeomTransform, Dir, Length, ResultTime, ResultPosition, ResultNormal, ResultFaceIndex, ResultFaceNormal, Thickness, bComputeMTD);
+			}
+			return bResult;
+		};
+	}
+
+	TEST_CASE("SweepSphereVsBox", "[Chaos][Sweeps][unit]")
+	{
+		const bool bComputeMTD = true;
+		const FReal Thickness = 0;
+		const FReal Length = 100;
+		const TBox<double, 3> TestGeom(FVec3(-1), FVec3(1));
+		const TSphere<double, 3> SweepGeom(FVec3::ZeroVector, 2);
+
+		FReal ResultTime;
+		FVec3 ResultPosition;
+		FVec3 ResultNormal;
+		FVec3 ResultFaceNormal = FVec3::ZeroVector;
+		int32 ResultFaceIndex = INDEX_NONE;
+
+		SECTION("Sweep From -X")
+		{
+			const FRigidTransform3 TestGeomTransform(FVec3(1, 2, 3), FRotation3::FromIdentity());
+			const FRigidTransform3 SweepGeomTransform(FVec3(-9, 2, 3), FRotation3::FromIdentity());
+			const FVec3 Dir(1, 0, 0);
+
+			const bool bResult = SweepQuery(TestGeom, TestGeomTransform, SweepGeom, SweepGeomTransform, Dir, Length, ResultTime, ResultPosition, ResultNormal, ResultFaceIndex, ResultFaceNormal, Thickness, bComputeMTD);
+			CHECK(bResult);
+			CHECK_THAT(ResultTime, Catch::Matchers::WithinRel(7.0, UE_DOUBLE_KINDA_SMALL_NUMBER));
+			CHECK_THAT(ResultPosition, Catch::ApproxEq(FVec3(0, 2, 3)));
+			CHECK_THAT(ResultNormal, Catch::ApproxEq(FVec3(-1, 0, 0)));
+			CHECK_THAT(ResultFaceNormal, Catch::ApproxEq(FVec3(-1, 0, 0)));
+			CHECK(ResultFaceIndex == INDEX_NONE);
+		}
+		SECTION("Sweep From +X")
+		{
+			const FRigidTransform3 TestGeomTransform(FVec3(1, 2, 3), FRotation3::FromIdentity());
+			const FRigidTransform3 SweepGeomTransform(FVec3(11, 2, 3), FRotation3::FromIdentity());
+			const FVec3 Dir(-1, 0, 0);
+
+			const bool bResult = SweepQuery(TestGeom, TestGeomTransform, SweepGeom, SweepGeomTransform, Dir, Length, ResultTime, ResultPosition, ResultNormal, ResultFaceIndex, ResultFaceNormal, Thickness, bComputeMTD);
+			CHECK(bResult);
+			CHECK_THAT(ResultTime, Catch::Matchers::WithinRel(7.0, UE_DOUBLE_KINDA_SMALL_NUMBER));
+			CHECK_THAT(ResultPosition, Catch::ApproxEq(FVec3(2, 2, 3)));
+			CHECK_THAT(ResultNormal, Catch::ApproxEq(FVec3(1, 0, 0)));
+			CHECK_THAT(ResultFaceNormal, Catch::ApproxEq(FVec3(1, 0, 0)));
+			CHECK(ResultFaceIndex == INDEX_NONE);
+		}
+		SECTION("Sweep From -Y")
+		{
+			const FRigidTransform3 TestGeomTransform(FVec3(1, 2, 3), FRotation3::FromIdentity());
+			const FRigidTransform3 SweepGeomTransform(FVec3(1, -8, 3), FRotation3::FromIdentity());
+			const FVec3 Dir(0, 1, 0);
+
+			const bool bResult = SweepQuery(TestGeom, TestGeomTransform, SweepGeom, SweepGeomTransform, Dir, Length, ResultTime, ResultPosition, ResultNormal, ResultFaceIndex, ResultFaceNormal, Thickness, bComputeMTD);
+			CHECK(bResult);
+			CHECK_THAT(ResultTime, Catch::Matchers::WithinRel(7.0, UE_DOUBLE_KINDA_SMALL_NUMBER));
+			CHECK_THAT(ResultPosition, Catch::ApproxEq(FVec3(1, 1, 3)));
+			CHECK_THAT(ResultNormal, Catch::ApproxEq(FVec3(0, -1, 0)));
+			CHECK_THAT(ResultFaceNormal, Catch::ApproxEq(FVec3(0, -1, 0)));
+			CHECK(ResultFaceIndex == INDEX_NONE);
+		}
+		SECTION("Sweep From +Y")
+		{
+			const FRigidTransform3 TestGeomTransform(FVec3(1, 2, 3), FRotation3::FromIdentity());
+			const FRigidTransform3 SweepGeomTransform(FVec3(1, 12, 3), FRotation3::FromIdentity());
+			const FVec3 Dir(0, -1, 0);
+
+			const bool bResult = SweepQuery(TestGeom, TestGeomTransform, SweepGeom, SweepGeomTransform, Dir, Length, ResultTime, ResultPosition, ResultNormal, ResultFaceIndex, ResultFaceNormal, Thickness, bComputeMTD);
+			CHECK(bResult);
+			CHECK_THAT(ResultTime, Catch::Matchers::WithinRel(7.0, UE_DOUBLE_KINDA_SMALL_NUMBER));
+			CHECK_THAT(ResultPosition, Catch::ApproxEq(FVec3(1, 3, 3)));
+			CHECK_THAT(ResultNormal, Catch::ApproxEq(FVec3(0, 1, 0)));
+			CHECK_THAT(ResultFaceNormal, Catch::ApproxEq(FVec3(0, 1, 0)));
+			CHECK(ResultFaceIndex == INDEX_NONE);
+		}
+		SECTION("Sweep From -Z")
+		{
+			const FRigidTransform3 TestGeomTransform(FVec3(1, 2, 3), FRotation3::FromIdentity());
+			const FRigidTransform3 SweepGeomTransform(FVec3(1, 2, -7), FRotation3::FromIdentity());
+			const FVec3 Dir(0, 0, 1);
+
+			const bool bResult = SweepQuery(TestGeom, TestGeomTransform, SweepGeom, SweepGeomTransform, Dir, Length, ResultTime, ResultPosition, ResultNormal, ResultFaceIndex, ResultFaceNormal, Thickness, bComputeMTD);
+			CHECK(bResult);
+			CHECK_THAT(ResultTime, Catch::Matchers::WithinRel(7.0, UE_DOUBLE_KINDA_SMALL_NUMBER));
+			CHECK_THAT(ResultPosition, Catch::ApproxEq(FVec3(1, 2, 2)));
+			CHECK_THAT(ResultNormal, Catch::ApproxEq(FVec3(0, 0, -1)));
+			CHECK_THAT(ResultFaceNormal, Catch::ApproxEq(FVec3(0, 0, -1)));
+			CHECK(ResultFaceIndex == INDEX_NONE);
+		}
+		SECTION("Sweep From +Z")
+		{
+			const FRigidTransform3 TestGeomTransform(FVec3(1, 2, 3), FRotation3::FromIdentity());
+			const FRigidTransform3 SweepGeomTransform(FVec3(1, 2, 13), FRotation3::FromIdentity());
+			const FVec3 Dir(0, 0, -1);
+
+			const bool bResult = SweepQuery(TestGeom, TestGeomTransform, SweepGeom, SweepGeomTransform, Dir, Length, ResultTime, ResultPosition, ResultNormal, ResultFaceIndex, ResultFaceNormal, Thickness, bComputeMTD);
+			CHECK(bResult);
+			CHECK_THAT(ResultTime, Catch::Matchers::WithinRel(7.0, UE_DOUBLE_KINDA_SMALL_NUMBER));
+			CHECK_THAT(ResultPosition, Catch::ApproxEq(FVec3(1, 2, 4)));
+			CHECK_THAT(ResultNormal, Catch::ApproxEq(FVec3(0, 0, 1)));
+			CHECK_THAT(ResultFaceNormal, Catch::ApproxEq(FVec3(0, 0, 1)));
+			CHECK(ResultFaceIndex == INDEX_NONE);
+		}
+		SECTION("Sweep: Past max distance")
+		{
+			const FRigidTransform3 TestGeomTransform(FVec3(1, 2, 3), FRotation3::FromIdentity());
+			const FRigidTransform3 SweepGeomTransform(FVec3(1.05 + Length + 3, 2, 3), FRotation3::FromIdentity());
+			const FVec3 Dir(-1, 0, 0);
+
+			const bool bResult = SweepQuery(TestGeom, TestGeomTransform, SweepGeom, SweepGeomTransform, Dir, Length, ResultTime, ResultPosition, ResultNormal, ResultFaceIndex, ResultFaceNormal, Thickness, bComputeMTD);
+			CHECK(!bResult);
+		}
+		SECTION("Sweep: Miss")
+		{
+			const FRigidTransform3 TestGeomTransform(FVec3(1, 2, 3), FRotation3::FromIdentity());
+			const FRigidTransform3 SweepGeomTransform(FVec3(-9, 10, 3), FRotation3::FromIdentity());
+			const FVec3 Dir(1, 0, 0);
+
+			const bool bResult = SweepQuery(TestGeom, TestGeomTransform, SweepGeom, SweepGeomTransform, Dir, Length, ResultTime, ResultPosition, ResultNormal, ResultFaceIndex, ResultFaceNormal, Thickness, bComputeMTD);
+			CHECK(!bResult);
+		}
+		SECTION("Shallow Initial Overlap")
+		{
+			const FRigidTransform3 TestGeomTransform(FVec3(1, 2, 3), FRotation3::FromIdentity());
+			const FRigidTransform3 SweepGeomTransform(FVec3(-1, 2, 3), FRotation3::FromIdentity());
+			const FVec3 Dir(1, 0, 0);
+
+			const bool bResult = SweepQuery(TestGeom, TestGeomTransform, SweepGeom, SweepGeomTransform, Dir, Length, ResultTime, ResultPosition, ResultNormal, ResultFaceIndex, ResultFaceNormal, Thickness, bComputeMTD);
+			CHECK(bResult);
+			CHECK_THAT(ResultTime, Catch::Matchers::WithinRel(-1.0, UE_DOUBLE_KINDA_SMALL_NUMBER));
+			CHECK_THAT(ResultPosition, Catch::ApproxEq(FVec3(0, 2, 3)));
+			CHECK_THAT(ResultNormal, Catch::ApproxEq(FVec3(-1, 0, 0)));
+			CHECK_THAT(ResultFaceNormal, Catch::ApproxEq(FVec3(-1, 0, 0)));
+			CHECK(ResultFaceIndex == INDEX_NONE);
+		}
+		SECTION("Deep Initial Overlap")
+		{
+			const FRigidTransform3 TestGeomTransform(FVec3(1, 2, 3), FRotation3::FromIdentity());
+			const FRigidTransform3 SweepGeomTransform(FVec3(0.9, 2, 3), FRotation3::FromIdentity());
+			const FVec3 Dir(1, 0, 0);
+
+			const bool bResult = SweepQuery(TestGeom, TestGeomTransform, SweepGeom, SweepGeomTransform, Dir, Length, ResultTime, ResultPosition, ResultNormal, ResultFaceIndex, ResultFaceNormal, Thickness, bComputeMTD);
+			CHECK(bResult);
+			CHECK_THAT(ResultTime, Catch::Matchers::WithinRel(-2.9, UE_DOUBLE_KINDA_SMALL_NUMBER));
+			CHECK_THAT(ResultPosition, Catch::ApproxEq(FVec3(0, 2, 3)));
+			CHECK_THAT(ResultNormal, Catch::ApproxEq(FVec3(-1, 0, 0)));
+			CHECK_THAT(ResultFaceNormal, Catch::ApproxEq(FVec3(-1, 0, 0)));
+			CHECK(ResultFaceIndex == INDEX_NONE);
+		}
+		SECTION("Initial Overlap Behind")
+		{
+			const FRigidTransform3 TestGeomTransform(FVec3(1, 2, 3), FRotation3::FromIdentity());
+			const FRigidTransform3 SweepGeomTransform(FVec3(2, 2, 3), FRotation3::FromIdentity());
+			const FVec3 Dir(1, 0, 0);
+
+			const bool bResult = SweepQuery(TestGeom, TestGeomTransform, SweepGeom, SweepGeomTransform, Dir, Length, ResultTime, ResultPosition, ResultNormal, ResultFaceIndex, ResultFaceNormal, Thickness, bComputeMTD);
+			CHECK(bResult);
+			CHECK_THAT(ResultTime, Catch::Matchers::WithinRel(-2, UE_DOUBLE_KINDA_SMALL_NUMBER));
+			CHECK_THAT(ResultPosition, Catch::ApproxEq(FVec3(2, 2, 3)));
+			CHECK_THAT(ResultNormal, Catch::ApproxEq(FVec3(1, 0, 0)));
+			CHECK_THAT(ResultFaceNormal, Catch::ApproxEq(FVec3(1, 0, 0)));
+			CHECK(ResultFaceIndex == INDEX_NONE);
+		}
+		SECTION("Rotation - X 45")
+		{
+			// Do a sweep along the local y axis, however rotate the entire test by 45 degrees along the x.
+			const FVec3 GeomLocation = FVec3(1, 2, 3);
+			const FRigidTransform3 TestGeomTransform(GeomLocation, FRotation3::FromAxisAngle(FVec3(1, 0, 0), PI * 0.25));
+			const FVec3 WorldOffset = TestGeomTransform.TransformVectorNoScale(FVec3(0, 10, 0));
+			const FRigidTransform3 SweepGeomTransform(GeomLocation + WorldOffset, FRotation3::FromIdentity());
+			const FVec3 Dir = -WorldOffset.GetSafeNormal();
+			const FVec3 ExpectedNormal = -Dir;
+			const FVec3 ExpectedPosition = GeomLocation + ExpectedNormal; // Box has size 1
+
+			const bool bResult = SweepQuery(TestGeom, TestGeomTransform, SweepGeom, SweepGeomTransform, Dir, Length, ResultTime, ResultPosition, ResultNormal, ResultFaceIndex, ResultFaceNormal, Thickness, bComputeMTD);
+			CHECK(bResult);
+			CHECK_THAT(ResultTime, Catch::Matchers::WithinRel(7, UE_DOUBLE_KINDA_SMALL_NUMBER));
+			CHECK_THAT(ResultPosition, Catch::ApproxEq(ExpectedPosition));
+			CHECK_THAT(ResultNormal, Catch::ApproxEq(ExpectedNormal));
+			CHECK_THAT(ResultFaceNormal, Catch::ApproxEq(ExpectedNormal));
+			CHECK(ResultFaceIndex == INDEX_NONE);
+		}
+		SECTION("Zero Length Sweep: Miss")
+		{
+			const FRigidTransform3 TestGeomTransform(FVec3(1, 2, 3), FRotation3::FromIdentity());
+			const FRigidTransform3 SweepGeomTransform(FVec3(-9, 2, 3), FRotation3::FromIdentity());
+			const FVec3 Dir(1, 0, 0);
+			const FReal TestLength = 0;
+
+			const bool bResult = SweepQuery(TestGeom, TestGeomTransform, SweepGeom, SweepGeomTransform, Dir, TestLength, ResultTime, ResultPosition, ResultNormal, ResultFaceIndex, ResultFaceNormal, Thickness, bComputeMTD);
+			CHECK(bResult == false);
+		}
+		SECTION("Zero Length Sweep: Hit")
+		{
+			const FRigidTransform3 TestGeomTransform(FVec3(1, 2, 3), FRotation3::FromIdentity());
+			const FRigidTransform3 SweepGeomTransform(FVec3(-1, 2, 3), FRotation3::FromIdentity());
+			const FVec3 Dir(1, 0, 0);
+			const FReal TestLength = 0;
+
+			const bool bResult = SweepQuery(TestGeom, TestGeomTransform, SweepGeom, SweepGeomTransform, Dir, TestLength, ResultTime, ResultPosition, ResultNormal, ResultFaceIndex, ResultFaceNormal, Thickness, bComputeMTD);
+			CHECK(bResult);
+			CHECK_THAT(ResultTime, Catch::Matchers::WithinRel(-1.0, UE_DOUBLE_KINDA_SMALL_NUMBER));
+			CHECK_THAT(ResultPosition, Catch::ApproxEq(FVec3(0, 2, 3)));
+			CHECK_THAT(ResultNormal, Catch::ApproxEq(FVec3(-1, 0, 0)));
+			CHECK_THAT(ResultFaceNormal, Catch::ApproxEq(FVec3(-1, 0, 0)));
+			CHECK(ResultFaceIndex == INDEX_NONE);
+		}
+		SECTION("Edge Intersection")
+		{
+			// This test primarily is to show the face normal is correctly computed and it's different than the normal
+			const FRigidTransform3 TestGeomTransform(FVec3(1, 2, 3), FRotation3::FromIdentity());
+			const FRigidTransform3 SweepGeomTransform(FVec3(3.99, 12, 3), FRotation3::FromIdentity());
+			const FVec3 Dir(0, -1, 0);
+
+			const bool bResult = SweepQuery(TestGeom, TestGeomTransform, SweepGeom, SweepGeomTransform, Dir, Length, ResultTime, ResultPosition, ResultNormal, ResultFaceIndex, ResultFaceNormal, Thickness, bComputeMTD);
+			CHECK(bResult);
+			CHECK_THAT(ResultTime, Catch::Matchers::WithinRel(8.8, 0.1));
+			CHECK_THAT(ResultPosition, Catch::ApproxEq(FVec3(2, 3, 3), 0.1));
+			CHECK_THAT(ResultNormal, Catch::ApproxEq(FVec3(1, 0, 0), 0.1));
+			CHECK_THAT(ResultFaceNormal, Catch::ApproxEq(FVec3(0, 1, 0)));
+			CHECK(ResultFaceIndex == INDEX_NONE);
+		}
+	}
+
+	TEST_CASE("SweepSphereVsBox - Perf", "[Chaos][Sweeps][!benchmark]")
+	{
+		const bool bComputeMTD = true;
+		const FReal Thickness = 0;
+		const FReal Length = 100;
+		const FVec3 Dir(1, 0, 0);
+		const TBox<double, 3> TestGeom(FVec3(-1), FVec3(1));
+		const TSphere<double, 3> SweepGeom(FVec3::ZeroVector, 2);
+		const FRigidTransform3 TestGeomTransform(FVec3(1, 2, 3), FRotation3::FromIdentity());
+
+		FReal ResultTime;
+		FVec3 ResultPosition, ResultNormal, ResultFaceNormal;
+		int32 ResultFaceIndex = INDEX_NONE;
+
+		BENCHMARK("Basic Hit")
+		{
+			const FRigidTransform3 SweepGeomTransform(FVec3(-9, 2, 3), FRotation3::FromIdentity());
+
+			bool bResult = false;
+			for (int32 I = 0; I < 1000; ++I)
+			{
+				bResult |= SweepQuery(TestGeom, TestGeomTransform, SweepGeom, SweepGeomTransform, Dir, Length, ResultTime, ResultPosition, ResultNormal, ResultFaceIndex, ResultFaceNormal, Thickness, bComputeMTD);
+			}
+			return bResult;
+		};
+		BENCHMARK("Edge Hit")
+		{
+			const FVec3 Dir2 = FVec3(-1, 0, -1).GetSafeNormal();
+			const FRigidTransform3 SweepGeomTransform(TestGeomTransform.GetLocation() - Dir2 * 10, FRotation3::FromIdentity());
+
+			bool bResult = false;
+			for (int32 I = 0; I < 1000; ++I)
+			{
+				bResult |= SweepQuery(TestGeom, TestGeomTransform, SweepGeom, SweepGeomTransform, Dir2, Length, ResultTime, ResultPosition, ResultNormal, ResultFaceIndex, ResultFaceNormal, Thickness, bComputeMTD);
+			}
+			return bResult;
+		};
+		BENCHMARK("Vertex Hit")
+		{
+			const FVec3 Dir2 = FVec3(-1, -1, -1).GetSafeNormal();
+			const FRigidTransform3 SweepGeomTransform(TestGeomTransform.GetLocation() - Dir2 * 10, FRotation3::FromIdentity());
+
+			bool bResult = false;
+			for (int32 I = 0; I < 1000; ++I)
+			{
+				bResult |= SweepQuery(TestGeom, TestGeomTransform, SweepGeom, SweepGeomTransform, Dir2, Length, ResultTime, ResultPosition, ResultNormal, ResultFaceIndex, ResultFaceNormal, Thickness, bComputeMTD);
+			}
+			return bResult;
+		};
+		BENCHMARK("Miss")
+		{
+			const FRigidTransform3 SweepGeomTransform(FVec3(-9, 10, 3), FRotation3::FromIdentity());
+
+			bool bResult = false;
+			for (int32 I = 0; I < 1000; ++I)
+			{
+				bResult |= SweepQuery(TestGeom, TestGeomTransform, SweepGeom, SweepGeomTransform, Dir, Length, ResultTime, ResultPosition, ResultNormal, ResultFaceIndex, ResultFaceNormal, Thickness, bComputeMTD);
+			}
+			return bResult;
+		};
+		BENCHMARK("Shallow Initial Overlap")
+		{
+			const FRigidTransform3 SweepGeomTransform(FVec3(-1, 2, 3), FRotation3::FromIdentity());
+
+			bool bResult = false;
+			for (int32 I = 0; I < 1000; ++I)
+			{
+				bResult |= SweepQuery(TestGeom, TestGeomTransform, SweepGeom, SweepGeomTransform, Dir, Length, ResultTime, ResultPosition, ResultNormal, ResultFaceIndex, ResultFaceNormal, Thickness, bComputeMTD);
+			}
+			return bResult;
+		};
+		BENCHMARK("Deep Initial Overlap")
+		{
+			const FRigidTransform3 SweepGeomTransform(FVec3(0.9, 2, 3), FRotation3::FromIdentity());
+
+			bool bResult = false;
+			for (int32 I = 0; I < 1000; ++I)
+			{
+				bResult |= SweepQuery(TestGeom, TestGeomTransform, SweepGeom, SweepGeomTransform, Dir, Length, ResultTime, ResultPosition, ResultNormal, ResultFaceIndex, ResultFaceNormal, Thickness, bComputeMTD);
+			}
+			return bResult;
+		};
+	}
+
+	TEST_CASE("SweepSphereVsCapsule", "[Chaos][Sweeps][unit]")
+	{
+		const bool bComputeMTD = true;
+		const FReal Thickness = 0;
+		const FReal Length = 100;
+		const FCapsule TestGeom(FVec3(0, -1, 0), FVec3(0, 1, 0), 3);
+		const TSphere<double, 3> SweepGeom(FVec3::ZeroVector, 2);
+
+		FReal ResultTime;
+		FVec3 ResultPosition;
+		FVec3 ResultNormal;
+		FVec3 ResultFaceNormal = FVec3::ZeroVector;
+		int32 ResultFaceIndex = INDEX_NONE;
+
+		SECTION("Sweep From -X")
+		{
+			const FRigidTransform3 TestGeomTransform(FVec3(1, 2, 3), FRotation3::FromIdentity());
+			const FRigidTransform3 SweepGeomTransform(FVec3(-9, 2, 3), FRotation3::FromIdentity());
+			const FVec3 Dir(1, 0, 0);
+
+			const bool bResult = SweepQuery(TestGeom, TestGeomTransform, SweepGeom, SweepGeomTransform, Dir, Length, ResultTime, ResultPosition, ResultNormal, ResultFaceIndex, ResultFaceNormal, Thickness, bComputeMTD);
+			CHECK(bResult);
+			CHECK_THAT(ResultTime, Catch::Matchers::WithinRel(5.0, UE_DOUBLE_KINDA_SMALL_NUMBER));
+			CHECK_THAT(ResultPosition, Catch::ApproxEq(FVec3(-2, 2, 3)));
+			CHECK_THAT(ResultNormal, Catch::ApproxEq(FVec3(-1, 0, 0)));
+			CHECK_THAT(ResultFaceNormal, Catch::ApproxEq(FVec3(-1, 0, 0)));
+			CHECK(ResultFaceIndex == INDEX_NONE);
+		}
+		SECTION("Sweep From +X")
+		{
+			const FRigidTransform3 TestGeomTransform(FVec3(1, 2, 3), FRotation3::FromIdentity());
+			const FRigidTransform3 SweepGeomTransform(FVec3(11, 2, 3), FRotation3::FromIdentity());
+			const FVec3 Dir(-1, 0, 0);
+
+			const bool bResult = SweepQuery(TestGeom, TestGeomTransform, SweepGeom, SweepGeomTransform, Dir, Length, ResultTime, ResultPosition, ResultNormal, ResultFaceIndex, ResultFaceNormal, Thickness, bComputeMTD);
+			CHECK(bResult);
+			CHECK_THAT(ResultTime, Catch::Matchers::WithinRel(5.0, UE_DOUBLE_KINDA_SMALL_NUMBER));
+			CHECK_THAT(ResultPosition, Catch::ApproxEq(FVec3(4, 2, 3)));
+			CHECK_THAT(ResultNormal, Catch::ApproxEq(FVec3(1, 0, 0)));
+			CHECK_THAT(ResultFaceNormal, Catch::ApproxEq(FVec3(1, 0, 0)));
+			CHECK(ResultFaceIndex == INDEX_NONE);
+		}
+		SECTION("Sweep From -Y")
+		{
+			const FRigidTransform3 TestGeomTransform(FVec3(1, 2, 3), FRotation3::FromIdentity());
+			const FRigidTransform3 SweepGeomTransform(FVec3(1, -8, 3), FRotation3::FromIdentity());
+			const FVec3 Dir(0, 1, 0);
+
+			const bool bResult = SweepQuery(TestGeom, TestGeomTransform, SweepGeom, SweepGeomTransform, Dir, Length, ResultTime, ResultPosition, ResultNormal, ResultFaceIndex, ResultFaceNormal, Thickness, bComputeMTD);
+			CHECK(bResult);
+			CHECK_THAT(ResultTime, Catch::Matchers::WithinRel(4.0, UE_DOUBLE_KINDA_SMALL_NUMBER));
+			CHECK_THAT(ResultPosition, Catch::ApproxEq(FVec3(1, -2, 3)));
+			CHECK_THAT(ResultNormal, Catch::ApproxEq(FVec3(0, -1, 0)));
+			CHECK_THAT(ResultFaceNormal, Catch::ApproxEq(FVec3(0, -1, 0)));
+			CHECK(ResultFaceIndex == INDEX_NONE);
+		}
+		SECTION("Sweep From +Y")
+		{
+			const FRigidTransform3 TestGeomTransform(FVec3(1, 2, 3), FRotation3::FromIdentity());
+			const FRigidTransform3 SweepGeomTransform(FVec3(1, 12, 3), FRotation3::FromIdentity());
+			const FVec3 Dir(0, -1, 0);
+
+			const bool bResult = SweepQuery(TestGeom, TestGeomTransform, SweepGeom, SweepGeomTransform, Dir, Length, ResultTime, ResultPosition, ResultNormal, ResultFaceIndex, ResultFaceNormal, Thickness, bComputeMTD);
+			CHECK(bResult);
+			CHECK_THAT(ResultTime, Catch::Matchers::WithinRel(4.0, UE_DOUBLE_KINDA_SMALL_NUMBER));
+			CHECK_THAT(ResultPosition, Catch::ApproxEq(FVec3(1, 6, 3)));
+			CHECK_THAT(ResultNormal, Catch::ApproxEq(FVec3(0, 1, 0)));
+			CHECK_THAT(ResultFaceNormal, Catch::ApproxEq(FVec3(0, 1, 0)));
+			CHECK(ResultFaceIndex == INDEX_NONE);
+		}
+		SECTION("Sweep From -Z")
+		{
+			const FRigidTransform3 TestGeomTransform(FVec3(1, 2, 3), FRotation3::FromIdentity());
+			const FRigidTransform3 SweepGeomTransform(FVec3(1, 2, -7), FRotation3::FromIdentity());
+			const FVec3 Dir(0, 0, 1);
+
+			const bool bResult = SweepQuery(TestGeom, TestGeomTransform, SweepGeom, SweepGeomTransform, Dir, Length, ResultTime, ResultPosition, ResultNormal, ResultFaceIndex, ResultFaceNormal, Thickness, bComputeMTD);
+			CHECK(bResult);
+			CHECK_THAT(ResultTime, Catch::Matchers::WithinRel(5.0, UE_DOUBLE_KINDA_SMALL_NUMBER));
+			CHECK_THAT(ResultPosition, Catch::ApproxEq(FVec3(1, 2, 0)));
+			CHECK_THAT(ResultNormal, Catch::ApproxEq(FVec3(0, 0, -1)));
+			CHECK_THAT(ResultFaceNormal, Catch::ApproxEq(FVec3(0, 0, -1)));
+			CHECK(ResultFaceIndex == INDEX_NONE);
+		}
+		SECTION("Sweep From +Z")
+		{
+			const FRigidTransform3 TestGeomTransform(FVec3(1, 2, 3), FRotation3::FromIdentity());
+			const FRigidTransform3 SweepGeomTransform(FVec3(1, 2, 13), FRotation3::FromIdentity());
+			const FVec3 Dir(0, 0, -1);
+
+			const bool bResult = SweepQuery(TestGeom, TestGeomTransform, SweepGeom, SweepGeomTransform, Dir, Length, ResultTime, ResultPosition, ResultNormal, ResultFaceIndex, ResultFaceNormal, Thickness, bComputeMTD);
+			CHECK(bResult);
+			CHECK_THAT(ResultTime, Catch::Matchers::WithinRel(5.0, UE_DOUBLE_KINDA_SMALL_NUMBER));
+			CHECK_THAT(ResultPosition, Catch::ApproxEq(FVec3(1, 2, 6)));
+			CHECK_THAT(ResultNormal, Catch::ApproxEq(FVec3(0, 0, 1)));
+			CHECK_THAT(ResultFaceNormal, Catch::ApproxEq(FVec3(0, 0, 1)));
+			CHECK(ResultFaceIndex == INDEX_NONE);
+		}
+		SECTION("Sweep From -X with diagonal hit")
+		{
+			// Use law of sines to figure out the positional offset for a 45 degree angle hit
+			const FReal ZOffset = (TestGeom.GetRadiusf() + SweepGeom.GetRadiusf()) * FMath::Sin(FMath::DegreesToRadians(45)) / FMath::Sin(FMath::DegreesToRadians(90));
+			const FRigidTransform3 TestGeomTransform(FVec3(1, 2, 3), FRotation3::FromIdentity());
+			const FRigidTransform3 SweepGeomTransform(FVec3(-9, 2, 3 + ZOffset), FRotation3::FromIdentity());
+			const FVec3 Dir(1, 0, 0);
+			const FVec3 ExpectedNormal = FVec3(-1, 0, 1).GetSafeNormal();
+			const FVec3 ExpectedPosition = TestGeomTransform.GetLocation() + ExpectedNormal * TestGeom.GetRadiusf();
+
+			const bool bResult = SweepQuery(TestGeom, TestGeomTransform, SweepGeom, SweepGeomTransform, Dir, Length, ResultTime, ResultPosition, ResultNormal, ResultFaceIndex, ResultFaceNormal, Thickness, bComputeMTD);
+			CHECK(bResult);
+			CHECK_THAT(ResultTime, Catch::Matchers::WithinRel(6.464466, UE_DOUBLE_KINDA_SMALL_NUMBER));
+			CHECK_THAT(ResultPosition, Catch::ApproxEq(ExpectedPosition));
+			CHECK_THAT(ResultNormal, Catch::ApproxEq(ExpectedNormal));
+			CHECK_THAT(ResultFaceNormal, Catch::ApproxEq(ExpectedNormal));
+			CHECK(ResultFaceIndex == INDEX_NONE);
+		}
+		SECTION("Sweep: Past max distance")
+		{
+			const FRigidTransform3 TestGeomTransform(FVec3(1, 2, 3), FRotation3::FromIdentity());
+			const FRigidTransform3 SweepGeomTransform(FVec3(1.05 + Length + 5, 2, 3), FRotation3::FromIdentity());
+			const FVec3 Dir(-1, 0, 0);
+
+			const bool bResult = SweepQuery(TestGeom, TestGeomTransform, SweepGeom, SweepGeomTransform, Dir, Length, ResultTime, ResultPosition, ResultNormal, ResultFaceIndex, ResultFaceNormal, Thickness, bComputeMTD);
+			CHECK(!bResult);
+		}
+		SECTION("Sweep: Miss")
+		{
+			const FRigidTransform3 TestGeomTransform(FVec3(1, 2, 3), FRotation3::FromIdentity());
+			const FRigidTransform3 SweepGeomTransform(FVec3(-9, 10, 3), FRotation3::FromIdentity());
+			const FVec3 Dir(1, 0, 0);
+
+			const bool bResult = SweepQuery(TestGeom, TestGeomTransform, SweepGeom, SweepGeomTransform, Dir, Length, ResultTime, ResultPosition, ResultNormal, ResultFaceIndex, ResultFaceNormal, Thickness, bComputeMTD);
+			CHECK(!bResult);
+		}
+		SECTION("Shallow Initial Overlap")
+		{
+			const FRigidTransform3 TestGeomTransform(FVec3(1, 2, 3), FRotation3::FromIdentity());
+			const FRigidTransform3 SweepGeomTransform(FVec3(-3, 2, 3), FRotation3::FromIdentity());
+			const FVec3 Dir(1, 0, 0);
+
+			const bool bResult = SweepQuery(TestGeom, TestGeomTransform, SweepGeom, SweepGeomTransform, Dir, Length, ResultTime, ResultPosition, ResultNormal, ResultFaceIndex, ResultFaceNormal, Thickness, bComputeMTD);
+			CHECK(bResult);
+			CHECK_THAT(ResultTime, Catch::Matchers::WithinRel(-1.0, UE_DOUBLE_KINDA_SMALL_NUMBER));
+			CHECK_THAT(ResultPosition, Catch::ApproxEq(FVec3(-2, 2, 3)));
+			CHECK_THAT(ResultNormal, Catch::ApproxEq(FVec3(-1, 0, 0)));
+			CHECK_THAT(ResultFaceNormal, Catch::ApproxEq(FVec3(-1, 0, 0)));
+			CHECK(ResultFaceIndex == INDEX_NONE);
+		}
+		SECTION("Deep Initial Overlap")
+		{
+			const FRigidTransform3 TestGeomTransform(FVec3(1, 2, 3), FRotation3::FromIdentity());
+			const FRigidTransform3 SweepGeomTransform(FVec3(0.9, 2, 3), FRotation3::FromIdentity());
+			const FVec3 Dir(1, 0, 0);
+
+			const bool bResult = SweepQuery(TestGeom, TestGeomTransform, SweepGeom, SweepGeomTransform, Dir, Length, ResultTime, ResultPosition, ResultNormal, ResultFaceIndex, ResultFaceNormal, Thickness, bComputeMTD);
+			CHECK(bResult);
+			CHECK_THAT(ResultTime, Catch::Matchers::WithinRel(-4.9, UE_DOUBLE_KINDA_SMALL_NUMBER));
+			CHECK_THAT(ResultPosition, Catch::ApproxEq(FVec3(-2, 2, 3)));
+			CHECK_THAT(ResultNormal, Catch::ApproxEq(FVec3(-1, 0, 0)));
+			CHECK_THAT(ResultFaceNormal, Catch::ApproxEq(FVec3(-1, 0, 0)));
+			CHECK(ResultFaceIndex == INDEX_NONE);
+		}
+		SECTION("Initial Overlap Behind")
+		{
+			const FRigidTransform3 TestGeomTransform(FVec3(1, 2, 3), FRotation3::FromIdentity());
+			const FRigidTransform3 SweepGeomTransform(FVec3(2, 2, 3), FRotation3::FromIdentity());
+			const FVec3 Dir(1, 0, 0);
+
+			const bool bResult = SweepQuery(TestGeom, TestGeomTransform, SweepGeom, SweepGeomTransform, Dir, Length, ResultTime, ResultPosition, ResultNormal, ResultFaceIndex, ResultFaceNormal, Thickness, bComputeMTD);
+			CHECK(bResult);
+			CHECK_THAT(ResultTime, Catch::Matchers::WithinRel(-4, UE_DOUBLE_KINDA_SMALL_NUMBER));
+			CHECK_THAT(ResultPosition, Catch::ApproxEq(FVec3(4, 2, 3)));
+			CHECK_THAT(ResultNormal, Catch::ApproxEq(FVec3(1, 0, 0)));
+			CHECK_THAT(ResultFaceNormal, Catch::ApproxEq(FVec3(1, 0, 0)));
+			CHECK(ResultFaceIndex == INDEX_NONE);
+		}
+		SECTION("Sphere Rotation: 90 about X")
+		{
+			const FRigidTransform3 TestGeomTransform(FVec3(1, 2, 3), FRotation3::FromIdentity());
+			const FRigidTransform3 SweepGeomTransform(FVec3(11, 2, 3), FRotation3::FromAxisAngle(FVec3(1, 0, 0), PI * 0.5));
+			const FVec3 Dir(-1, 0, 0);
+
+			const bool bResult = SweepQuery(TestGeom, TestGeomTransform, SweepGeom, SweepGeomTransform, Dir, Length, ResultTime, ResultPosition, ResultNormal, ResultFaceIndex, ResultFaceNormal, Thickness, bComputeMTD);
+			CHECK(bResult);
+			CHECK_THAT(ResultTime, Catch::Matchers::WithinRel(5.0, UE_DOUBLE_KINDA_SMALL_NUMBER));
+			CHECK_THAT(ResultPosition, Catch::ApproxEq(FVec3(4, 2, 3)));
+			CHECK_THAT(ResultNormal, Catch::ApproxEq(FVec3(1, 0, 0)));
+			CHECK_THAT(ResultFaceNormal, Catch::ApproxEq(FVec3(1, 0, 0)));
+			CHECK(ResultFaceIndex == INDEX_NONE);
+		}
+		SECTION("Capsule Rotation: 90 about X")
+		{
+			// Rotate 90 degrees along the x then do a sweep along the z axis. This is now the long axis of the capsule
+			const FRigidTransform3 TestGeomTransform(FVec3(1, 2, 3), FRotation3::FromAxisAngle(FVec3(1, 0, 0), PI * 0.5));
+			const FRigidTransform3 SweepGeomTransform(FVec3(1, 2, -7), FRotation3::FromIdentity());
+			const FVec3 Dir(0, 0, 1);
+
+			const bool bResult = SweepQuery(TestGeom, TestGeomTransform, SweepGeom, SweepGeomTransform, Dir, Length, ResultTime, ResultPosition, ResultNormal, ResultFaceIndex, ResultFaceNormal, Thickness, bComputeMTD);
+			CHECK(bResult);
+			CHECK_THAT(ResultTime, Catch::Matchers::WithinRel(4, UE_DOUBLE_KINDA_SMALL_NUMBER));
+			CHECK_THAT(ResultPosition, Catch::ApproxEq(FVec3(1, 2, -1)));
+			CHECK_THAT(ResultNormal, Catch::ApproxEq(FVec3(0, 0, -1)));
+			CHECK_THAT(ResultFaceNormal, Catch::ApproxEq(FVec3(0, 0, -1)));
+			CHECK(ResultFaceIndex == INDEX_NONE);
+		}
+		SECTION("Capsule Rotation: 45 about X")
+		{
+			// Do a sweep along the local y axis, however rotate the entire test by 45 degrees along the x.
+			const FVec3 GeomLocation = FVec3(1, 2, 3);
+			const FRigidTransform3 TestGeomTransform(GeomLocation, FRotation3::FromAxisAngle(FVec3(1, 0, 0), PI * 0.25));
+			const FVec3 WorldOffset = TestGeomTransform.TransformVectorNoScale(FVec3(0, 10, 0));
+			const FRigidTransform3 SweepGeomTransform(GeomLocation + WorldOffset, FRotation3::FromIdentity());
+			const FVec3 Dir = -WorldOffset.GetSafeNormal();
+			const FVec3 ExpectedNormal = -Dir;
+			const FVec3 ExpectedPosition = GeomLocation + ExpectedNormal * 4; // Offset on axis is radius + half extent = 3 + 1 = 4
+
+			const bool bResult = SweepQuery(TestGeom, TestGeomTransform, SweepGeom, SweepGeomTransform, Dir, Length, ResultTime, ResultPosition, ResultNormal, ResultFaceIndex, ResultFaceNormal, Thickness, bComputeMTD);
+			CHECK(bResult);
+			CHECK_THAT(ResultTime, Catch::Matchers::WithinRel(4.0, UE_DOUBLE_KINDA_SMALL_NUMBER));
+			CHECK_THAT(ResultPosition, Catch::ApproxEq(ExpectedPosition));
+			CHECK_THAT(ResultNormal, Catch::ApproxEq(ExpectedNormal));
+			CHECK_THAT(ResultFaceNormal, Catch::ApproxEq(ExpectedNormal));
+			CHECK(ResultFaceIndex == INDEX_NONE);
+		}
+		SECTION("With Thickness")
+		{
+			const FRigidTransform3 TestGeomTransform(FVec3(1, 2, 3), FRotation3::FromIdentity());
+			const FRigidTransform3 SweepGeomTransform(FVec3(-9, 2, 3), FRotation3::FromIdentity());
+			const FVec3 Dir(1, 0, 0);
+			const FReal TestThickness = 1;
+
+			const bool bResult = SweepQuery(TestGeom, TestGeomTransform, SweepGeom, SweepGeomTransform, Dir, Length, ResultTime, ResultPosition, ResultNormal, ResultFaceIndex, ResultFaceNormal, TestThickness, bComputeMTD);
+			CHECK(bResult);
+			CHECK_THAT(ResultTime, Catch::Matchers::WithinRel(4.0, UE_DOUBLE_KINDA_SMALL_NUMBER));
+			CHECK_THAT(ResultPosition, Catch::ApproxEq(FVec3(-2, 2, 3)));
+			CHECK_THAT(ResultNormal, Catch::ApproxEq(FVec3(-1, 0, 0)));
+			CHECK_THAT(ResultFaceNormal, Catch::ApproxEq(FVec3(-1, 0, 0)));
+			CHECK(ResultFaceIndex == INDEX_NONE);
+		}
+		SECTION("Zero Length Sweep: Miss")
+		{
+			const FRigidTransform3 TestGeomTransform(FVec3(1, 2, 3), FRotation3::FromIdentity());
+			const FRigidTransform3 SweepGeomTransform(FVec3(-9, 2, 3), FRotation3::FromIdentity());
+			const FVec3 Dir(1, 0, 0);
+			const FReal TestLength = 0;
+
+			const bool bResult = SweepQuery(TestGeom, TestGeomTransform, SweepGeom, SweepGeomTransform, Dir, TestLength, ResultTime, ResultPosition, ResultNormal, ResultFaceIndex, ResultFaceNormal, Thickness, bComputeMTD);
+			CHECK(bResult == false);
+		}
+		SECTION("Zero Length Sweep: Hit")
+		{
+			const FRigidTransform3 TestGeomTransform(FVec3(1, 2, 3), FRotation3::FromIdentity());
+			const FRigidTransform3 SweepGeomTransform(FVec3(-3, 2, 3), FRotation3::FromIdentity());
+			const FVec3 Dir(1, 0, 0);
+			const FReal TestLength = 0;
+
+			const bool bResult = SweepQuery(TestGeom, TestGeomTransform, SweepGeom, SweepGeomTransform, Dir, TestLength, ResultTime, ResultPosition, ResultNormal, ResultFaceIndex, ResultFaceNormal, Thickness, bComputeMTD);
+			CHECK(bResult);
+			CHECK_THAT(ResultTime, Catch::Matchers::WithinRel(-1.0, UE_DOUBLE_KINDA_SMALL_NUMBER));
+			CHECK_THAT(ResultPosition, Catch::ApproxEq(FVec3(-2, 2, 3)));
+			CHECK_THAT(ResultNormal, Catch::ApproxEq(FVec3(-1, 0, 0)));
+			CHECK_THAT(ResultFaceNormal, Catch::ApproxEq(FVec3(-1, 0, 0)));
+			CHECK(ResultFaceIndex == INDEX_NONE);
+		}
+	}
+
+	TEST_CASE("SweepSphereVsCapsule - Perf", "[Chaos][Sweeps][!benchmark]")
+	{
+		const bool bComputeMTD = true;
+		const FVec3 Dir(1, 0, 0);
+		const FReal Length = 100;
+		const FReal Thickness = 0;
+		const FCapsule TestGeom(FVec3(0, -1, 0), FVec3(0, 1, 0), 3);
+		const TSphere<double, 3> SweepGeom(FVec3::ZeroVector, 2);
+		const FRigidTransform3 TestGeomTransform(FVec3(1, 2, 3), FRotation3::FromIdentity());
+
+		FReal ResultTime;
+		FVec3 ResultPosition, ResultNormal, ResultFaceNormal;
+		int32 ResultFaceIndex = INDEX_NONE;
+
+		BENCHMARK("Basic Hit")
+		{
+			const FRigidTransform3 SweepGeomTransform(FVec3(-9, 2, 3), FRotation3::FromIdentity());
+
+			bool bResult = false;
+			for (int32 I = 0; I < 1000; ++I)
+			{
+				bResult |= SweepQuery(TestGeom, TestGeomTransform, SweepGeom, SweepGeomTransform, Dir, Length, ResultTime, ResultPosition, ResultNormal, ResultFaceIndex, ResultFaceNormal, Thickness, bComputeMTD);
+			}
+			return bResult;
+		};
+		BENCHMARK("Miss")
+		{
+			const FRigidTransform3 SweepGeomTransform(FVec3(-9, 10, 3), FRotation3::FromIdentity());
+
+			bool bResult = false;
+			for (int32 I = 0; I < 1000; ++I)
+			{
+				bResult |= SweepQuery(TestGeom, TestGeomTransform, SweepGeom, SweepGeomTransform, Dir, Length, ResultTime, ResultPosition, ResultNormal, ResultFaceIndex, ResultFaceNormal, Thickness, bComputeMTD);
+			}
+			return bResult;
+		};
+		BENCHMARK("Shallow Initial Overlap")
+		{
+			const FRigidTransform3 SweepGeomTransform(FVec3(-3, 2, 3), FRotation3::FromIdentity());
+
+			bool bResult = false;
+			for (int32 I = 0; I < 1000; ++I)
+			{
+				bResult |= SweepQuery(TestGeom, TestGeomTransform, SweepGeom, SweepGeomTransform, Dir, Length, ResultTime, ResultPosition, ResultNormal, ResultFaceIndex, ResultFaceNormal, Thickness, bComputeMTD);
+			}
+			return bResult;
+		};
+		BENCHMARK("Deep Initial Overlap")
+		{
+			const FRigidTransform3 SweepGeomTransform(FVec3(0.9, 2, 3), FRotation3::FromIdentity());
+
+			bool bResult = false;
+			for (int32 I = 0; I < 1000; ++I)
+			{
+				bResult |= SweepQuery(TestGeom, TestGeomTransform, SweepGeom, SweepGeomTransform, Dir, Length, ResultTime, ResultPosition, ResultNormal, ResultFaceIndex, ResultFaceNormal, Thickness, bComputeMTD);
+			}
+			return bResult;
+		};
+	}
+
+	TEST_CASE("SweepBoxVsSphere", "[Chaos][Sweeps][unit]")
+	{
+		const bool bComputeMTD = true;
+		const FReal Thickness = 0;
+		const FReal Length = 100;
+		const TSphere<double, 3> TestGeom(FVec3::ZeroVector, 2);
+		const TBox<double, 3> SweepGeom(FVec3(-1), FVec3(1));
+
+		FReal ResultTime;
+		FVec3 ResultPosition;
+		FVec3 ResultNormal;
+		FVec3 ResultFaceNormal = FVec3::ZeroVector;
+		int32 ResultFaceIndex = INDEX_NONE;
+
+		SECTION("Sweep From -X")
+		{
+			const FRigidTransform3 TestGeomTransform(FVec3(1, 2, 3), FRotation3::FromIdentity());
+			const FRigidTransform3 SweepGeomTransform(FVec3(-9, 2, 3), FRotation3::FromIdentity());
+			const FVec3 Dir(1, 0, 0);
+
+			const bool bResult = SweepQuery(TestGeom, TestGeomTransform, SweepGeom, SweepGeomTransform, Dir, Length, ResultTime, ResultPosition, ResultNormal, ResultFaceIndex, ResultFaceNormal, Thickness, bComputeMTD);
+			CHECK(bResult);
+			CHECK_THAT(ResultTime, Catch::Matchers::WithinRel(7.0, UE_DOUBLE_KINDA_SMALL_NUMBER));
+			CHECK_THAT(ResultPosition, Catch::ApproxEq(FVec3(-1, 2, 3)));
+			CHECK_THAT(ResultNormal, Catch::ApproxEq(FVec3(-1, 0, 0)));
+			CHECK_THAT(ResultFaceNormal, Catch::ApproxEq(FVec3(-1, 0, 0)));
+			CHECK(ResultFaceIndex == INDEX_NONE);
+		}
+		SECTION("Sweep: Past max distance")
+		{
+			const FRigidTransform3 TestGeomTransform(FVec3(1, 2, 3), FRotation3::FromIdentity());
+			const FRigidTransform3 SweepGeomTransform(FVec3(1.05 + Length + 3, 2, 3), FRotation3::FromIdentity());
+			const FVec3 Dir(-1, 0, 0);
+
+			const bool bResult = SweepQuery(TestGeom, TestGeomTransform, SweepGeom, SweepGeomTransform, Dir, Length, ResultTime, ResultPosition, ResultNormal, ResultFaceIndex, ResultFaceNormal, Thickness, bComputeMTD);
+			CHECK(!bResult);
+		}
+		SECTION("Sweep: Miss")
+		{
+			const FRigidTransform3 TestGeomTransform(FVec3(1, 2, 3), FRotation3::FromIdentity());
+			const FRigidTransform3 SweepGeomTransform(FVec3(-9, 10, 3), FRotation3::FromIdentity());
+			const FVec3 Dir(1, 0, 0);
+
+			const bool bResult = SweepQuery(TestGeom, TestGeomTransform, SweepGeom, SweepGeomTransform, Dir, Length, ResultTime, ResultPosition, ResultNormal, ResultFaceIndex, ResultFaceNormal, Thickness, bComputeMTD);
+			CHECK(!bResult);
+		}
+		SECTION("Shallow Initial Overlap")
+		{
+			const FRigidTransform3 TestGeomTransform(FVec3(1, 2, 3), FRotation3::FromIdentity());
+			const FRigidTransform3 SweepGeomTransform(FVec3(-1, 2, 3), FRotation3::FromIdentity());
+			const FVec3 Dir(1, 0, 0);
+
+			const bool bResult = SweepQuery(TestGeom, TestGeomTransform, SweepGeom, SweepGeomTransform, Dir, Length, ResultTime, ResultPosition, ResultNormal, ResultFaceIndex, ResultFaceNormal, Thickness, bComputeMTD);
+			CHECK(bResult);
+			CHECK_THAT(ResultTime, Catch::Matchers::WithinRel(-1.0, UE_DOUBLE_KINDA_SMALL_NUMBER));
+			CHECK_THAT(ResultPosition, Catch::ApproxEq(FVec3(-1, 2, 3)));
+			CHECK_THAT(ResultNormal, Catch::ApproxEq(FVec3(-1, 0, 0)));
+			CHECK_THAT(ResultFaceNormal, Catch::ApproxEq(FVec3(-1, 0, 0)));
+			CHECK(ResultFaceIndex == INDEX_NONE);
+		}
+		SECTION("Edge Intersection")
+		{
+			// This test primarily is to show the face normal is correctly computed and it's different than the normal
+			const FRigidTransform3 TestGeomTransform(FVec3(1, 2, 3), FRotation3::FromIdentity());
+			const FRigidTransform3 SweepGeomTransform(FVec3(3.99, 12, 3), FRotation3::FromIdentity());
+			const FVec3 Dir(0, -1, 0);
+
+			const bool bResult = SweepQuery(TestGeom, TestGeomTransform, SweepGeom, SweepGeomTransform, Dir, Length, ResultTime, ResultPosition, ResultNormal, ResultFaceIndex, ResultFaceNormal, Thickness, bComputeMTD);
+			CHECK(bResult);
+			CHECK_THAT(ResultTime, Catch::Matchers::WithinRel(8.8, 0.1));
+			CHECK_THAT(ResultPosition, Catch::ApproxEq(FVec3(3, 2.2, 3), 0.1));
+			CHECK_THAT(ResultNormal, Catch::ApproxEq(FVec3(1, 0, 0), 0.1));
+			CHECK_THAT(ResultFaceNormal, Catch::ApproxEq(FVec3(1, 0, 0), 0.1));
+			CHECK(ResultFaceIndex == INDEX_NONE);
+		}
+	}
+
+	TEST_CASE("SweepBoxVsBox", "[Chaos][Sweeps][unit]")
+	{
+		// TODO
+	}
+
+	TEST_CASE("SweepBoxVsBox - Perf", "[Chaos][Sweeps][!benchmark]")
+	{
+		const bool bComputeMTD = true;
+		const FReal Thickness = 0;
+		const FReal Length = 100;
+		const TBox<double, 3> TestGeom(FVec3(-1), FVec3(1));
+		const TBox<double, 3> SweepGeom(FVec3(-1), FVec3(1));
+		const FRigidTransform3 TestGeomTransform(FVec3(1, 2, 3), FRotation3::FromIdentity());
+		const FVec3 Dir(1, 0, 0);
+
+		FReal ResultTime;
+		FVec3 ResultPosition;
+		FVec3 ResultNormal;
+		FVec3 ResultFaceNormal = FVec3::ZeroVector;
+		int32 ResultFaceIndex = INDEX_NONE;
+
+		BENCHMARK("Basic Hit")
+		{
+			const FRigidTransform3 SweepGeomTransform(FVec3(-9, 2, 3), FRotation3::FromIdentity());
+
+			bool bResult = false;
+			for (int32 I = 0; I < 1000; ++I)
+			{
+				bResult |= SweepQuery(TestGeom, TestGeomTransform, SweepGeom, SweepGeomTransform, Dir, Length, ResultTime, ResultPosition, ResultNormal, ResultFaceIndex, ResultFaceNormal, Thickness, bComputeMTD);
+			}
+			return bResult;
+		};
+		BENCHMARK("Miss")
+		{
+			const FRigidTransform3 SweepGeomTransform(FVec3(-9, 10, 3), FRotation3::FromIdentity());
+
+			bool bResult = false;
+			for (int32 I = 0; I < 1000; ++I)
+			{
+				bResult |= SweepQuery(TestGeom, TestGeomTransform, SweepGeom, SweepGeomTransform, Dir, Length, ResultTime, ResultPosition, ResultNormal, ResultFaceIndex, ResultFaceNormal, Thickness, bComputeMTD);
+			}
+			return bResult;
+		};
+		BENCHMARK("Shallow Initial Overlap")
+		{
+			const FRigidTransform3 SweepGeomTransform(FVec3(-2.5, 2, 3), FRotation3::FromIdentity());
+
+			bool bResult = false;
+			for (int32 I = 0; I < 1000; ++I)
+			{
+				bResult |= SweepQuery(TestGeom, TestGeomTransform, SweepGeom, SweepGeomTransform, Dir, Length, ResultTime, ResultPosition, ResultNormal, ResultFaceIndex, ResultFaceNormal, Thickness, bComputeMTD);
+			}
+			return bResult;
+		};
+		BENCHMARK("Deep Initial Overlap")
+		{
+			const FRigidTransform3 SweepGeomTransform(FVec3(0.9, 2, 3), FRotation3::FromIdentity());
+
+			bool bResult = false;
+			for (int32 I = 0; I < 1000; ++I)
+			{
+				bResult |= SweepQuery(TestGeom, TestGeomTransform, SweepGeom, SweepGeomTransform, Dir, Length, ResultTime, ResultPosition, ResultNormal, ResultFaceIndex, ResultFaceNormal, Thickness, bComputeMTD);
+			}
+			return bResult;
+		};
+	}
+
+	TEST_CASE("SweepBoxVsCapsule", "[Chaos][Sweeps][unit]")
+	{
+		const bool bComputeMTD = true;
+		const FReal Thickness = 0;
+		const FReal Length = 100;
+		const FCapsule XCapsule(FVec3(-1, 0, 0), FVec3(1, 0, 0), 2);
+		const FCapsule YCapsule(FVec3(0, -1, 0), FVec3(0, 1, 0), 2);
+		const FCapsule ZCapsule(FVec3(0, 0, -1), FVec3(0, 0, 1), 2);
+		const TBox<double, 3> SweepGeom(FVec3(-1), FVec3(1));
+
+		FReal ResultTime;
+		FVec3 ResultPosition;
+		FVec3 ResultNormal;
+		FVec3 ResultFaceNormal = FVec3::ZeroVector;
+		int32 ResultFaceIndex = INDEX_NONE;
+
+		SECTION("Sweep From -X")
+		{
+			const FRigidTransform3 TestGeomTransform(FVec3(1, 2, 3), FRotation3::FromIdentity());
+			const FRigidTransform3 SweepGeomTransform(FVec3(-9, 2, 3), FRotation3::FromIdentity());
+			const FVec3 Dir(1, 0, 0);
+
+			const bool bResult = SweepQuery(XCapsule, TestGeomTransform, SweepGeom, SweepGeomTransform, Dir, Length, ResultTime, ResultPosition, ResultNormal, ResultFaceIndex, ResultFaceNormal, Thickness, bComputeMTD);
+			CHECK(bResult);
+			CHECK_THAT(ResultTime, Catch::Matchers::WithinRel(6.0, UE_DOUBLE_KINDA_SMALL_NUMBER));
+			CHECK_THAT(ResultPosition, Catch::ApproxEq(FVec3(-2, 2, 3)));
+			CHECK_THAT(ResultNormal, Catch::ApproxEq(FVec3(-1, 0, 0)));
+			CHECK_THAT(ResultFaceNormal, Catch::ApproxEq(FVec3(-1, 0, 0)));
+			CHECK(ResultFaceIndex == INDEX_NONE);
+		}
+		SECTION("Sweep From +X")
+		{
+			const FRigidTransform3 TestGeomTransform(FVec3(1, 2, 3), FRotation3::FromIdentity());
+			const FRigidTransform3 SweepGeomTransform(FVec3(11, 2, 3), FRotation3::FromIdentity());
+			const FVec3 Dir(-1, 0, 0);
+
+			const bool bResult = SweepQuery(XCapsule, TestGeomTransform, SweepGeom, SweepGeomTransform, Dir, Length, ResultTime, ResultPosition, ResultNormal, ResultFaceIndex, ResultFaceNormal, Thickness, bComputeMTD);
+			CHECK(bResult);
+			CHECK_THAT(ResultTime, Catch::Matchers::WithinRel(6.0, UE_DOUBLE_KINDA_SMALL_NUMBER));
+			CHECK_THAT(ResultPosition, Catch::ApproxEq(FVec3(4, 2, 3)));
+			CHECK_THAT(ResultNormal, Catch::ApproxEq(FVec3(1, 0, 0)));
+			CHECK_THAT(ResultFaceNormal, Catch::ApproxEq(FVec3(1, 0, 0)));
+			CHECK(ResultFaceIndex == INDEX_NONE);
+		}
+		SECTION("Sweep From -Y")
+		{
+			const FRigidTransform3 TestGeomTransform(FVec3(1, 2, 3), FRotation3::FromIdentity());
+			const FRigidTransform3 SweepGeomTransform(FVec3(1, -8, 3), FRotation3::FromIdentity());
+			const FVec3 Dir(0, 1, 0);
+
+			const bool bResult = SweepQuery(YCapsule, TestGeomTransform, SweepGeom, SweepGeomTransform, Dir, Length, ResultTime, ResultPosition, ResultNormal, ResultFaceIndex, ResultFaceNormal, Thickness, bComputeMTD);
+			CHECK(bResult);
+			CHECK_THAT(ResultTime, Catch::Matchers::WithinRel(6.0, UE_DOUBLE_KINDA_SMALL_NUMBER));
+			CHECK_THAT(ResultPosition, Catch::ApproxEq(FVec3(1, -1, 3)));
+			CHECK_THAT(ResultNormal, Catch::ApproxEq(FVec3(0, -1, 0)));
+			CHECK_THAT(ResultFaceNormal, Catch::ApproxEq(FVec3(0, -1, 0)));
+			CHECK(ResultFaceIndex == INDEX_NONE);
+		}
+		SECTION("Sweep From +Y")
+		{
+			const FRigidTransform3 TestGeomTransform(FVec3(1, 2, 3), FRotation3::FromIdentity());
+			const FRigidTransform3 SweepGeomTransform(FVec3(1, 12, 3), FRotation3::FromIdentity());
+			const FVec3 Dir(0, -1, 0);
+
+			const bool bResult = SweepQuery(YCapsule, TestGeomTransform, SweepGeom, SweepGeomTransform, Dir, Length, ResultTime, ResultPosition, ResultNormal, ResultFaceIndex, ResultFaceNormal, Thickness, bComputeMTD);
+			CHECK(bResult);
+			CHECK_THAT(ResultTime, Catch::Matchers::WithinRel(6.0, UE_DOUBLE_KINDA_SMALL_NUMBER));
+			CHECK_THAT(ResultPosition, Catch::ApproxEq(FVec3(1, 5, 3)));
+			CHECK_THAT(ResultNormal, Catch::ApproxEq(FVec3(0, 1, 0)));
+			CHECK_THAT(ResultFaceNormal, Catch::ApproxEq(FVec3(0, 1, 0)));
+			CHECK(ResultFaceIndex == INDEX_NONE);
+		}
+		SECTION("Sweep From -Z")
+		{
+			const FRigidTransform3 TestGeomTransform(FVec3(1, 2, 3), FRotation3::FromIdentity());
+			const FRigidTransform3 SweepGeomTransform(FVec3(1, 2, -7), FRotation3::FromIdentity());
+			const FVec3 Dir(0, 0, 1);
+
+			const bool bResult = SweepQuery(ZCapsule, TestGeomTransform, SweepGeom, SweepGeomTransform, Dir, Length, ResultTime, ResultPosition, ResultNormal, ResultFaceIndex, ResultFaceNormal, Thickness, bComputeMTD);
+			CHECK(bResult);
+			CHECK_THAT(ResultTime, Catch::Matchers::WithinRel(6.0, UE_DOUBLE_KINDA_SMALL_NUMBER));
+			CHECK_THAT(ResultPosition, Catch::ApproxEq(FVec3(1, 2, 0)));
+			CHECK_THAT(ResultNormal, Catch::ApproxEq(FVec3(0, 0, -1)));
+			CHECK_THAT(ResultFaceNormal, Catch::ApproxEq(FVec3(0, 0, -1)));
+			CHECK(ResultFaceIndex == INDEX_NONE);
+		}
+		SECTION("Sweep From +Z")
+		{
+			const FRigidTransform3 TestGeomTransform(FVec3(1, 2, 3), FRotation3::FromIdentity());
+			const FRigidTransform3 SweepGeomTransform(FVec3(1, 2, 13), FRotation3::FromIdentity());
+			const FVec3 Dir(0, 0, -1);
+
+			const bool bResult = SweepQuery(ZCapsule, TestGeomTransform, SweepGeom, SweepGeomTransform, Dir, Length, ResultTime, ResultPosition, ResultNormal, ResultFaceIndex, ResultFaceNormal, Thickness, bComputeMTD);
+			CHECK(bResult);
+			CHECK_THAT(ResultTime, Catch::Matchers::WithinRel(6.0, UE_DOUBLE_KINDA_SMALL_NUMBER));
+			CHECK_THAT(ResultPosition, Catch::ApproxEq(FVec3(1, 2, 6)));
+			CHECK_THAT(ResultNormal, Catch::ApproxEq(FVec3(0, 0, 1)));
+			CHECK_THAT(ResultFaceNormal, Catch::ApproxEq(FVec3(0, 0, 1)));
+			CHECK(ResultFaceIndex == INDEX_NONE);
+		}
+		SECTION("Sweep: Past max distance")
+		{
+			const FRigidTransform3 TestGeomTransform(FVec3(1, 2, 3), FRotation3::FromIdentity());
+			const FRigidTransform3 SweepGeomTransform(FVec3(1.05 + Length + 7, 2, 3), FRotation3::FromIdentity());
+			const FVec3 Dir(-1, 0, 0);
+
+			const bool bResult = SweepQuery(XCapsule, TestGeomTransform, SweepGeom, SweepGeomTransform, Dir, Length, ResultTime, ResultPosition, ResultNormal, ResultFaceIndex, ResultFaceNormal, Thickness, bComputeMTD);
+			CHECK(!bResult);
+		}
+		SECTION("Sweep: Miss")
+		{
+			const FRigidTransform3 TestGeomTransform(FVec3(1, 2, 3), FRotation3::FromIdentity());
+			const FRigidTransform3 SweepGeomTransform(FVec3(-9, 10, 3), FRotation3::FromIdentity());
+			const FVec3 Dir(1, 0, 0);
+
+			const bool bResult = SweepQuery(XCapsule, TestGeomTransform, SweepGeom, SweepGeomTransform, Dir, Length, ResultTime, ResultPosition, ResultNormal, ResultFaceIndex, ResultFaceNormal, Thickness, bComputeMTD);
+			CHECK(!bResult);
+		}
+		SECTION("Shallow Initial Overlap")
+		{
+			const FRigidTransform3 TestGeomTransform(FVec3(1, 2, 3), FRotation3::FromIdentity());
+			const FRigidTransform3 SweepGeomTransform(FVec3(-2, 2, 3), FRotation3::FromIdentity());
+			const FVec3 Dir(1, 0, 0);
+
+			const bool bResult = SweepQuery(XCapsule, TestGeomTransform, SweepGeom, SweepGeomTransform, Dir, Length, ResultTime, ResultPosition, ResultNormal, ResultFaceIndex, ResultFaceNormal, Thickness, bComputeMTD);
+			CHECK(bResult);
+			CHECK_THAT(ResultTime, Catch::Matchers::WithinRel(-1.0, UE_DOUBLE_KINDA_SMALL_NUMBER));
+			CHECK_THAT(ResultPosition, Catch::ApproxEq(FVec3(-2, 2, 3)));
+			CHECK_THAT(ResultNormal, Catch::ApproxEq(FVec3(-1, 0, 0)));
+			CHECK_THAT(ResultFaceNormal, Catch::ApproxEq(FVec3(-1, 0, 0)));
+			CHECK(ResultFaceIndex == INDEX_NONE);
+		}
+		SECTION("Deep Initial Overlap")
+		{
+			const FCapsule SmallXCapsule(FVec3(-1, 0, 0), FVec3(1, 0, 0), 0.5);
+			const FRigidTransform3 TestGeomTransform(FVec3(1, 2, 3), FRotation3::FromIdentity());
+			const FRigidTransform3 SweepGeomTransform(FVec3(-0.9, 2, 3), FRotation3::FromIdentity());
+			const FVec3 Dir(1, 0, 0);
+
+			const bool bResult = SweepQuery(SmallXCapsule, TestGeomTransform, SweepGeom, SweepGeomTransform, Dir, Length, ResultTime, ResultPosition, ResultNormal, ResultFaceIndex, ResultFaceNormal, Thickness, bComputeMTD);
+			CHECK(bResult);
+			CHECK_THAT(ResultTime, Catch::Matchers::WithinRel(-0.6, UE_DOUBLE_KINDA_SMALL_NUMBER));
+			CHECK_THAT(ResultPosition, Catch::ApproxEq(FVec3(-0.5, 2, 3)));
+			CHECK_THAT(ResultNormal, Catch::ApproxEq(FVec3(-1, 0, 0)));
+			CHECK_THAT(ResultFaceNormal, Catch::ApproxEq(FVec3(-1, 0, 0)));
+			CHECK(ResultFaceIndex == INDEX_NONE);
+		}
+		SECTION("Zero Length Sweep: Miss")
+		{
+			const FRigidTransform3 TestGeomTransform(FVec3(1, 2, 3), FRotation3::FromIdentity());
+			const FRigidTransform3 SweepGeomTransform(FVec3(-9, 2, 3), FRotation3::FromIdentity());
+			const FVec3 Dir(1, 0, 0);
+			const FReal TestLength = 0;
+
+			const bool bResult = SweepQuery(XCapsule, TestGeomTransform, SweepGeom, SweepGeomTransform, Dir, TestLength, ResultTime, ResultPosition, ResultNormal, ResultFaceIndex, ResultFaceNormal, Thickness, bComputeMTD);
+			CHECK(bResult == false);
+		}
+		SECTION("Zero Length Sweep: Hit")
+		{
+			const FRigidTransform3 TestGeomTransform(FVec3(1, 2, 3), FRotation3::FromIdentity());
+			const FRigidTransform3 SweepGeomTransform(FVec3(-2, 2, 3), FRotation3::FromIdentity());
+			const FVec3 Dir(1, 0, 0);
+			const FReal TestLength = 0;
+
+			const bool bResult = SweepQuery(XCapsule, TestGeomTransform, SweepGeom, SweepGeomTransform, Dir, TestLength, ResultTime, ResultPosition, ResultNormal, ResultFaceIndex, ResultFaceNormal, Thickness, bComputeMTD);
+			CHECK(bResult);
+			CHECK_THAT(ResultTime, Catch::Matchers::WithinRel(-1.0, UE_DOUBLE_KINDA_SMALL_NUMBER));
+			CHECK_THAT(ResultPosition, Catch::ApproxEq(FVec3(-2, 2, 3)));
+			CHECK_THAT(ResultNormal, Catch::ApproxEq(FVec3(-1, 0, 0)));
+			CHECK_THAT(ResultFaceNormal, Catch::ApproxEq(FVec3(-1, 0, 0)));
+			CHECK(ResultFaceIndex == INDEX_NONE);
+		}
+		SECTION("Edge Intersection")
+		{
+			// This test primarily is to show the face normal is correctly computed and it's different than the normal
+			const FRigidTransform3 TestGeomTransform(FVec3(1, 2, 3), FRotation3::FromIdentity());
+			const FRigidTransform3 SweepGeomTransform(FVec3(3.99, 12, 3), FRotation3::FromIdentity());
+			const FVec3 Dir(0, -1, 0);
+
+			const bool bResult = SweepQuery(YCapsule, TestGeomTransform, SweepGeom, SweepGeomTransform, Dir, Length, ResultTime, ResultPosition, ResultNormal, ResultFaceIndex, ResultFaceNormal, Thickness, bComputeMTD);
+			CHECK(bResult);
+			CHECK_THAT(ResultTime, Catch::Matchers::WithinRel(7.8, 0.1));
+			CHECK_THAT(ResultPosition, Catch::ApproxEq(FVec3(3, 3.2, 3), 0.1));
+			CHECK_THAT(ResultNormal, Catch::ApproxEq(FVec3(1, 0, 0), 0.1));
+			CHECK_THAT(ResultFaceNormal, Catch::ApproxEq(FVec3(1, 0, 0), 0.1));
+			CHECK(ResultFaceIndex == INDEX_NONE);
+		}
+		SECTION("Rotation")
+		{
+			// Rotate the box 45 degrees so we get an edge intersection and the capsule 90 degrees
+			const FRigidTransform3 TestGeomTransform(FVec3(1, 2, 3), FRotation3::FromAxisAngle(FVec3(0, 0, 1), PI * 0.5));
+			const FRigidTransform3 SweepGeomTransform(FVec3(-9, 2, 3), FRotation3::FromAxisAngle(FVec3(0, 0, 1), PI * 0.25));
+			const FVec3 Dir(1, 0, 0);
+
+			const bool bResult = SweepQuery(XCapsule, TestGeomTransform, SweepGeom, SweepGeomTransform, Dir, Length, ResultTime, ResultPosition, ResultNormal, ResultFaceIndex, ResultFaceNormal, Thickness, bComputeMTD);
+			CHECK(bResult);
+			CHECK_THAT(ResultTime, Catch::Matchers::WithinRel(6.58578, 0.01)); // 10 - radius(2) - sqrt(1 + 1) ~= 6.58
+			CHECK_THAT(ResultPosition, Catch::ApproxEq(FVec3(-1, 2, 3)));
+			CHECK_THAT(ResultNormal, Catch::ApproxEq(FVec3(-1, 0, 0)));
+			CHECK_THAT(ResultFaceNormal, Catch::ApproxEq(FVec3(-1, 0, 0)));
+			CHECK(ResultFaceIndex == INDEX_NONE);
+		}
+	}
+
+	TEST_CASE("SweepBoxVsCapsule - Perf", "[Chaos][Sweeps][!benchmark]")
+	{
+		const bool bComputeMTD = true;
+		const FReal Thickness = 0;
+		const FReal Length = 100;
+		const FCapsule XCapsule(FVec3(-1, 0, 0), FVec3(1, 0, 0), 2);
+		const TBox<double, 3> SweepGeom(FVec3(-1), FVec3(1));
+		const FRigidTransform3 TestGeomTransform(FVec3(1, 2, 3), FRotation3::FromIdentity());
+		const FVec3 Dir(1, 0, 0);
+
+		FReal ResultTime;
+		FVec3 ResultPosition;
+		FVec3 ResultNormal;
+		FVec3 ResultFaceNormal = FVec3::ZeroVector;
+		int32 ResultFaceIndex = INDEX_NONE;
+
+		BENCHMARK("Basic Hit")
+		{
+			const FRigidTransform3 SweepGeomTransform(FVec3(-9, 2, 3), FRotation3::FromIdentity());
+
+			bool bResult = false;
+			for (int32 I = 0; I < 1000; ++I)
+			{
+				bResult |= SweepQuery(XCapsule, TestGeomTransform, SweepGeom, SweepGeomTransform, Dir, Length, ResultTime, ResultPosition, ResultNormal, ResultFaceIndex, ResultFaceNormal, Thickness, bComputeMTD);
+			}
+			return bResult;
+		};
+		BENCHMARK("Miss")
+		{
+			const FRigidTransform3 SweepGeomTransform(FVec3(-9, 10, 3), FRotation3::FromIdentity());
+
+			bool bResult = false;
+			for (int32 I = 0; I < 1000; ++I)
+			{
+				bResult |= SweepQuery(XCapsule, TestGeomTransform, SweepGeom, SweepGeomTransform, Dir, Length, ResultTime, ResultPosition, ResultNormal, ResultFaceIndex, ResultFaceNormal, Thickness, bComputeMTD);
+			}
+			return bResult;
+		};
+		BENCHMARK("Shallow Initial Overlap")
+		{
+			const FRigidTransform3 SweepGeomTransform(FVec3(-2, 2, 3), FRotation3::FromIdentity());
+
+			bool bResult = false;
+			for (int32 I = 0; I < 1000; ++I)
+			{
+				bResult |= SweepQuery(XCapsule, TestGeomTransform, SweepGeom, SweepGeomTransform, Dir, Length, ResultTime, ResultPosition, ResultNormal, ResultFaceIndex, ResultFaceNormal, Thickness, bComputeMTD);
+			}
+			return bResult;
+		};
+		BENCHMARK("Deep Initial Overlap")
+		{
+			const FCapsule SmallXCapsule(FVec3(-1, 0, 0), FVec3(1, 0, 0), 0.5);
+			const FRigidTransform3 SweepGeomTransform(FVec3(-0.9, 2, 3), FRotation3::FromIdentity());
+
+			bool bResult = false;
+			for (int32 I = 0; I < 1000; ++I)
+			{
+				bResult |= SweepQuery(SmallXCapsule, TestGeomTransform, SweepGeom, SweepGeomTransform, Dir, Length, ResultTime, ResultPosition, ResultNormal, ResultFaceIndex, ResultFaceNormal, Thickness, bComputeMTD);
+			}
+			return bResult;
+		};
+	}
+
+	TEST_CASE("SweepCapsuleVsSphere", "[Chaos][Sweeps][unit]")
+	{
+		// Note: This sweep is minimally tested as it's just a flipped version of SphereVsCapsule. This is included to make sure the basics are correct.
+		const bool bComputeMTD = true;
+		const FReal Thickness = 0;
+		const FReal Length = 100;
+		const TSphere<double, 3> TestGeom(FVec3::ZeroVector, 2);
+		const FCapsule SweepGeom(FVec3(0, -1, 0), FVec3(0, 1, 0), 3);
+
+		FReal ResultTime;
+		FVec3 ResultPosition;
+		FVec3 ResultNormal;
+		FVec3 ResultFaceNormal = FVec3::ZeroVector;
+		int32 ResultFaceIndex = INDEX_NONE;
+
+		SECTION("Sweep From -X")
+		{
+			const FRigidTransform3 TestGeomTransform(FVec3(1, 2, 3), FRotation3::FromIdentity());
+			const FRigidTransform3 SweepGeomTransform(FVec3(-9, 2, 3), FRotation3::FromIdentity());
+			const FVec3 Dir(1, 0, 0);
+
+			const bool bResult = SweepQuery(TestGeom, TestGeomTransform, SweepGeom, SweepGeomTransform, Dir, Length, ResultTime, ResultPosition, ResultNormal, ResultFaceIndex, ResultFaceNormal, Thickness, bComputeMTD);
+			CHECK(bResult);
+			CHECK_THAT(ResultTime, Catch::Matchers::WithinRel(5.0, UE_DOUBLE_KINDA_SMALL_NUMBER));
+			CHECK_THAT(ResultPosition, Catch::ApproxEq(FVec3(-1, 2, 3)));
+			CHECK_THAT(ResultNormal, Catch::ApproxEq(FVec3(-1, 0, 0)));
+			CHECK_THAT(ResultFaceNormal, Catch::ApproxEq(FVec3(-1, 0, 0)));
+			CHECK(ResultFaceIndex == INDEX_NONE);
+		}
+		SECTION("Sweep From -X with diagonal hit")
+		{
+			// Use law of sines to figure out the positional offset for a 45 degree angle hit
+			const FReal ZOffset = (TestGeom.GetRadiusf() + SweepGeom.GetRadiusf()) * FMath::Sin(FMath::DegreesToRadians(45)) / FMath::Sin(FMath::DegreesToRadians(90));
+			const FRigidTransform3 TestGeomTransform(FVec3(1, 2, 3), FRotation3::FromIdentity());
+			const FRigidTransform3 SweepGeomTransform(FVec3(-9, 2, 3 + ZOffset), FRotation3::FromIdentity());
+			const FVec3 Dir(1, 0, 0);
+			const FVec3 ExpectedNormal = FVec3(-1, 0, 1).GetSafeNormal();
+			const FVec3 ExpectedPosition = TestGeomTransform.GetLocation() + ExpectedNormal * TestGeom.GetRadiusf();
+
+			const bool bResult = SweepQuery(TestGeom, TestGeomTransform, SweepGeom, SweepGeomTransform, Dir, Length, ResultTime, ResultPosition, ResultNormal, ResultFaceIndex, ResultFaceNormal, Thickness, bComputeMTD);
+			CHECK(bResult);
+			CHECK_THAT(ResultTime, Catch::Matchers::WithinRel(6.464466, UE_DOUBLE_KINDA_SMALL_NUMBER));
+			CHECK_THAT(ResultPosition, Catch::ApproxEq(ExpectedPosition));
+			CHECK_THAT(ResultNormal, Catch::ApproxEq(ExpectedNormal));
+			CHECK_THAT(ResultFaceNormal, Catch::ApproxEq(ExpectedNormal));
+			CHECK(ResultFaceIndex == INDEX_NONE);
+		}
+		SECTION("Sphere Rotation: 90 about X")
+		{
+			// Test a rotation on the above -X test. Nothing should change
+			const FRigidTransform3 TestGeomTransform(FVec3(1, 2, 3), FRotation3::FromAxisAngle(FVec3(1, 0, 0), PI * 0.5));
+			const FRigidTransform3 SweepGeomTransform(FVec3(-9, 2, 3), FRotation3::FromIdentity());
+			const FVec3 Dir(1, 0, 0);
+
+			const bool bResult = SweepQuery(TestGeom, TestGeomTransform, SweepGeom, SweepGeomTransform, Dir, Length, ResultTime, ResultPosition, ResultNormal, ResultFaceIndex, ResultFaceNormal, Thickness, bComputeMTD);
+			CHECK(bResult);
+			CHECK_THAT(ResultTime, Catch::Matchers::WithinRel(5.0, UE_DOUBLE_KINDA_SMALL_NUMBER));
+			CHECK_THAT(ResultPosition, Catch::ApproxEq(FVec3(-1, 2, 3)));
+			CHECK_THAT(ResultNormal, Catch::ApproxEq(FVec3(-1, 0, 0)));
+			CHECK_THAT(ResultFaceNormal, Catch::ApproxEq(FVec3(-1, 0, 0)));
+			CHECK(ResultFaceIndex == INDEX_NONE);
+		}
+		SECTION("Capsule Rotation: 90 about X")
+		{
+			// Rotate 90 degrees along the x then do a sweep along the z axis. This is now the long axis of the capsule
+			const FRigidTransform3 TestGeomTransform(FVec3(1, 2, 3), FRotation3::FromIdentity());
+			const FRigidTransform3 SweepGeomTransform(FVec3(1, 2, -7), FRotation3::FromAxisAngle(FVec3(1, 0, 0), PI * 0.5));
+			const FVec3 Dir(0, 0, 1);
+
+			const bool bResult = SweepQuery(TestGeom, TestGeomTransform, SweepGeom, SweepGeomTransform, Dir, Length, ResultTime, ResultPosition, ResultNormal, ResultFaceIndex, ResultFaceNormal, Thickness, bComputeMTD);
+			CHECK(bResult);
+			CHECK_THAT(ResultTime, Catch::Matchers::WithinRel(4, UE_DOUBLE_KINDA_SMALL_NUMBER));
+			CHECK_THAT(ResultPosition, Catch::ApproxEq(FVec3(1, 2, 1)));
+			CHECK_THAT(ResultNormal, Catch::ApproxEq(FVec3(0, 0, -1)));
+			CHECK_THAT(ResultFaceNormal, Catch::ApproxEq(FVec3(0, 0, -1)));
+			CHECK(ResultFaceIndex == INDEX_NONE);
+		}
+		SECTION("Sweep: Miss (Behind Sweep)")
+		{
+			const FRigidTransform3 TestGeomTransform(FVec3(1, 2, 3), FRotation3::FromIdentity());
+			const FRigidTransform3 SweepGeomTransform(FVec3(7, 10, 3), FRotation3::FromIdentity());
+			const FVec3 Dir(1, 0, 0);
+
+			const bool bResult = SweepQuery(TestGeom, TestGeomTransform, SweepGeom, SweepGeomTransform, Dir, Length, ResultTime, ResultPosition, ResultNormal, ResultFaceIndex, ResultFaceNormal, Thickness, bComputeMTD);
+			CHECK(!bResult);
+		}
+		SECTION("Shallow Initial Overlap")
+		{
+			const FRigidTransform3 TestGeomTransform(FVec3(1, 2, 3), FRotation3::FromIdentity());
+			const FRigidTransform3 SweepGeomTransform(FVec3(-3, 2, 3), FRotation3::FromIdentity());
+			const FVec3 Dir(1, 0, 0);
+
+			const bool bResult = SweepQuery(TestGeom, TestGeomTransform, SweepGeom, SweepGeomTransform, Dir, Length, ResultTime, ResultPosition, ResultNormal, ResultFaceIndex, ResultFaceNormal, Thickness, bComputeMTD);
+			CHECK(bResult);
+			CHECK_THAT(ResultTime, Catch::Matchers::WithinRel(-1.0, UE_DOUBLE_KINDA_SMALL_NUMBER));
+			CHECK_THAT(ResultPosition, Catch::ApproxEq(FVec3(-1, 2, 3)));
+			CHECK_THAT(ResultNormal, Catch::ApproxEq(FVec3(-1, 0, 0)));
+			CHECK_THAT(ResultFaceNormal, Catch::ApproxEq(FVec3(-1, 0, 0)));
+			CHECK(ResultFaceIndex == INDEX_NONE);
+		}
+		SECTION("Zero Length Sweep: Miss")
+		{
+			const FRigidTransform3 TestGeomTransform(FVec3(1, 2, 3), FRotation3::FromIdentity());
+			const FRigidTransform3 SweepGeomTransform(FVec3(-9, 2, 3), FRotation3::FromIdentity());
+			const FVec3 Dir(1, 0, 0);
+			const FReal TestLength = 0;
+
+			const bool bResult = SweepQuery(TestGeom, TestGeomTransform, SweepGeom, SweepGeomTransform, Dir, TestLength, ResultTime, ResultPosition, ResultNormal, ResultFaceIndex, ResultFaceNormal, Thickness, bComputeMTD);
+			CHECK(bResult == false);
+		}
+		SECTION("Zero Length Sweep: Hit")
+		{
+			const FRigidTransform3 TestGeomTransform(FVec3(1, 2, 3), FRotation3::FromIdentity());
+			const FRigidTransform3 SweepGeomTransform(FVec3(-3, 2, 3), FRotation3::FromIdentity());
+			const FVec3 Dir(1, 0, 0);
+			const FReal TestLength = 0;
+
+			const bool bResult = SweepQuery(TestGeom, TestGeomTransform, SweepGeom, SweepGeomTransform, Dir, TestLength, ResultTime, ResultPosition, ResultNormal, ResultFaceIndex, ResultFaceNormal, Thickness, bComputeMTD);
+			CHECK(bResult);
+			CHECK_THAT(ResultTime, Catch::Matchers::WithinRel(-1.0, UE_DOUBLE_KINDA_SMALL_NUMBER));
+			CHECK_THAT(ResultPosition, Catch::ApproxEq(FVec3(-1, 2, 3)));
+			CHECK_THAT(ResultNormal, Catch::ApproxEq(FVec3(-1, 0, 0)));
+			CHECK_THAT(ResultFaceNormal, Catch::ApproxEq(FVec3(-1, 0, 0)));
+			CHECK(ResultFaceIndex == INDEX_NONE);
+		}
+	}
+
+	TEST_CASE("SweepCapsuleVsSphere - Perf", "[Chaos][Sweeps][!benchmark]")
+	{
+		const bool bComputeMTD = true;
+		const FVec3 Dir(1, 0, 0);
+		const FReal Length = 100;
+		const FReal Thickness = 0;
+		const FCapsule SweepGeom(FVec3(0, -1, 0), FVec3(0, 1, 0), 3);
+		const TSphere<double, 3> TestGeom(FVec3::ZeroVector, 2);
+		const FRigidTransform3 TestGeomTransform(FVec3(1, 2, 3), FRotation3::FromIdentity());
+
+		FReal ResultTime;
+		FVec3 ResultPosition, ResultNormal, ResultFaceNormal;
+		int32 ResultFaceIndex = INDEX_NONE;
+
+		BENCHMARK("Basic Hit")
+		{
+			const FRigidTransform3 SweepGeomTransform(FVec3(-9, 2, 3), FRotation3::FromIdentity());
+
+			bool bResult = false;
+			for (int32 I = 0; I < 1000; ++I)
+			{
+				bResult |= SweepQuery(TestGeom, TestGeomTransform, SweepGeom, SweepGeomTransform, Dir, Length, ResultTime, ResultPosition, ResultNormal, ResultFaceIndex, ResultFaceNormal, Thickness, bComputeMTD);
+			}
+			return bResult;
+		};
+		BENCHMARK("Miss")
+		{
+			const FRigidTransform3 SweepGeomTransform(FVec3(-9, 10, 3), FRotation3::FromIdentity());
+
+			bool bResult = false;
+			for (int32 I = 0; I < 1000; ++I)
+			{
+				bResult |= SweepQuery(TestGeom, TestGeomTransform, SweepGeom, SweepGeomTransform, Dir, Length, ResultTime, ResultPosition, ResultNormal, ResultFaceIndex, ResultFaceNormal, Thickness, bComputeMTD);
+			}
+			return bResult;
+		};
+		BENCHMARK("Shallow Initial Overlap")
+		{
+			const FRigidTransform3 SweepGeomTransform(FVec3(-3, 2, 3), FRotation3::FromIdentity());
+
+			bool bResult = false;
+			for (int32 I = 0; I < 1000; ++I)
+			{
+				bResult |= SweepQuery(TestGeom, TestGeomTransform, SweepGeom, SweepGeomTransform, Dir, Length, ResultTime, ResultPosition, ResultNormal, ResultFaceIndex, ResultFaceNormal, Thickness, bComputeMTD);
+			}
+			return bResult;
+		};
+		BENCHMARK("Deep Initial Overlap")
+		{
+			const FRigidTransform3 SweepGeomTransform(FVec3(0.9, 2, 3), FRotation3::FromIdentity());
+
+			bool bResult = false;
+			for (int32 I = 0; I < 1000; ++I)
+			{
+				bResult |= SweepQuery(TestGeom, TestGeomTransform, SweepGeom, SweepGeomTransform, Dir, Length, ResultTime, ResultPosition, ResultNormal, ResultFaceIndex, ResultFaceNormal, Thickness, bComputeMTD);
+			}
+			return bResult;
+		};
+	}
+
+	TEST_CASE("SweepCapsuleVsBox", "[Chaos][Sweeps][unit]")
+	{
+		const bool bComputeMTD = true;
+		const FReal Thickness = 0;
+		const FReal Length = 100;
+		const FCapsule XCapsule(FVec3(-1, 0, 0), FVec3(1, 0, 0), 2);
+		const FCapsule YCapsule(FVec3(0, -1, 0), FVec3(0, 1, 0), 2);
+		const FCapsule ZCapsule(FVec3(0, 0, -1), FVec3(0, 0, 1), 2);
+		const TBox<double, 3> TestGeom(FVec3(-1), FVec3(1));
+
+		FReal ResultTime;
+		FVec3 ResultPosition;
+		FVec3 ResultNormal;
+		FVec3 ResultFaceNormal = FVec3::ZeroVector;
+		int32 ResultFaceIndex = INDEX_NONE;
+
+		SECTION("Sweep From -X")
+		{
+			const FRigidTransform3 TestGeomTransform(FVec3(1, 2, 3), FRotation3::FromIdentity());
+			const FRigidTransform3 SweepGeomTransform(FVec3(-9, 2, 3), FRotation3::FromIdentity());
+			const FVec3 Dir(1, 0, 0);
+
+			const bool bResult = SweepQuery(TestGeom, TestGeomTransform, XCapsule, SweepGeomTransform, Dir, Length, ResultTime, ResultPosition, ResultNormal, ResultFaceIndex, ResultFaceNormal, Thickness, bComputeMTD);
+			CHECK(bResult);
+			CHECK_THAT(ResultTime, Catch::Matchers::WithinRel(6.0, UE_DOUBLE_KINDA_SMALL_NUMBER));
+			CHECK_THAT(ResultPosition, Catch::ApproxEq(FVec3(0, 2, 3)));
+			CHECK_THAT(ResultNormal, Catch::ApproxEq(FVec3(-1, 0, 0)));
+			CHECK_THAT(ResultFaceNormal, Catch::ApproxEq(FVec3(-1, 0, 0)));
+			CHECK(ResultFaceIndex == INDEX_NONE);
+		}
+		SECTION("Sweep: Past max distance")
+		{
+			const FRigidTransform3 TestGeomTransform(FVec3(1, 2, 3), FRotation3::FromIdentity());
+			const FRigidTransform3 SweepGeomTransform(FVec3(1.05 + Length + 7, 2, 3), FRotation3::FromIdentity());
+			const FVec3 Dir(-1, 0, 0);
+
+			const bool bResult = SweepQuery(TestGeom, TestGeomTransform, XCapsule, SweepGeomTransform, Dir, Length, ResultTime, ResultPosition, ResultNormal, ResultFaceIndex, ResultFaceNormal, Thickness, bComputeMTD);
+			CHECK(!bResult);
+		}
+		SECTION("Sweep: Miss")
+		{
+			const FRigidTransform3 TestGeomTransform(FVec3(1, 2, 3), FRotation3::FromIdentity());
+			const FRigidTransform3 SweepGeomTransform(FVec3(-9, 10, 3), FRotation3::FromIdentity());
+			const FVec3 Dir(1, 0, 0);
+
+			const bool bResult = SweepQuery(TestGeom, TestGeomTransform, XCapsule, SweepGeomTransform, Dir, Length, ResultTime, ResultPosition, ResultNormal, ResultFaceIndex, ResultFaceNormal, Thickness, bComputeMTD);
+			CHECK(!bResult);
+		}
+		SECTION("Shallow Initial Overlap")
+		{
+			const FRigidTransform3 TestGeomTransform(FVec3(1, 2, 3), FRotation3::FromIdentity());
+			const FRigidTransform3 SweepGeomTransform(FVec3(-2, 2, 3), FRotation3::FromIdentity());
+			const FVec3 Dir(1, 0, 0);
+
+			const bool bResult = SweepQuery(TestGeom, TestGeomTransform, XCapsule, SweepGeomTransform, Dir, Length, ResultTime, ResultPosition, ResultNormal, ResultFaceIndex, ResultFaceNormal, Thickness, bComputeMTD);
+			CHECK(bResult);
+			CHECK_THAT(ResultTime, Catch::Matchers::WithinRel(-1.0, UE_DOUBLE_KINDA_SMALL_NUMBER));
+			CHECK_THAT(ResultPosition, Catch::ApproxEq(FVec3(0, 2, 3)));
+			CHECK_THAT(ResultNormal, Catch::ApproxEq(FVec3(-1, 0, 0)));
+			CHECK_THAT(ResultFaceNormal, Catch::ApproxEq(FVec3(-1, 0, 0)));
+			CHECK(ResultFaceIndex == INDEX_NONE);
+		}
+		SECTION("Edge Intersection")
+		{
+			// This test primarily is to show the face normal is correctly computed and it's different than the normal
+			const FRigidTransform3 TestGeomTransform(FVec3(1, 2, 3), FRotation3::FromIdentity());
+			const FRigidTransform3 SweepGeomTransform(FVec3(3.99, 12, 3), FRotation3::FromIdentity());
+			const FVec3 Dir(0, -1, 0);
+
+			const bool bResult = SweepQuery(TestGeom, TestGeomTransform, YCapsule, SweepGeomTransform, Dir, Length, ResultTime, ResultPosition, ResultNormal, ResultFaceIndex, ResultFaceNormal, Thickness, bComputeMTD);
+			CHECK(bResult);
+			CHECK_THAT(ResultTime, Catch::Matchers::WithinRel(7.8, 0.1));
+			CHECK_THAT(ResultPosition, Catch::ApproxEq(FVec3(2, 3, 3), 0.1));
+			CHECK_THAT(ResultNormal, Catch::ApproxEq(FVec3(1, 0, 0), 0.1));
+			CHECK_THAT(ResultFaceNormal, Catch::ApproxEq(FVec3(0, 1, 0), 0.1));
+			CHECK(ResultFaceIndex == INDEX_NONE);
+		}
+	}
+
+	TEST_CASE("SweepCapsuleVsCapsule", "[Chaos][Sweeps][unit]")
+	{
+		const bool bComputeMTD = true;
+		const FReal Thickness = 0;
+		const FReal Length = 100;
+		const FCapsule XCapsule(FVec3(-1, 0, 0), FVec3(1, 0, 0), 2);
+		const FCapsule YCapsule(FVec3(0, -1, 0), FVec3(0, 1, 0), 2);
+		const FCapsule ZCapsule(FVec3(0, 0, -1), FVec3(0, 0, 1), 2);
+
+		FReal ResultTime;
+		FVec3 ResultPosition;
+		FVec3 ResultNormal;
+		FVec3 ResultFaceNormal = FVec3::ZeroVector;
+		int32 ResultFaceIndex = INDEX_NONE;
+
+		SECTION("Parallel Edge Edge")
+		{
+			const FRigidTransform3 TestGeomTransform(FVec3(1, 2, 3), FRotation3::FromIdentity());
+			const FRigidTransform3 SweepGeomTransform(FVec3(-9, 2, 3), FRotation3::FromIdentity());
+			const FVec3 Dir(1, 0, 0);
+
+			const bool bResult = SweepQuery(YCapsule, TestGeomTransform, YCapsule, SweepGeomTransform, Dir, Length, ResultTime, ResultPosition, ResultNormal, ResultFaceIndex, ResultFaceNormal, Thickness, bComputeMTD);
+			CHECK(bResult);
+			CHECK_THAT(ResultTime, Catch::Matchers::WithinRel(6.0, UE_DOUBLE_KINDA_SMALL_NUMBER));
+			CHECK_THAT(ResultPosition.X, Catch::Matchers::WithinRel(-1, UE_DOUBLE_KINDA_SMALL_NUMBER));
+			CHECK_THAT(ResultNormal, Catch::ApproxEq(FVec3(-1, 0, 0)));
+			CHECK_THAT(ResultFaceNormal, Catch::ApproxEq(FVec3(-1, 0, 0)));
+			CHECK(ResultFaceIndex == INDEX_NONE);
+		}
+		SECTION("Perpendicular Edge Edge")
+		{
+			const FRigidTransform3 TestGeomTransform(FVec3(1, 2, 3), FRotation3::FromIdentity());
+			const FRigidTransform3 SweepGeomTransform(FVec3(-9, 2, 3), FRotation3::FromIdentity());
+			const FVec3 Dir(1, 0, 0);
+
+			const bool bResult = SweepQuery(YCapsule, TestGeomTransform, ZCapsule, SweepGeomTransform, Dir, Length, ResultTime, ResultPosition, ResultNormal, ResultFaceIndex, ResultFaceNormal, Thickness, bComputeMTD);
+			CHECK(bResult);
+			CHECK_THAT(ResultTime, Catch::Matchers::WithinRel(6.0, UE_DOUBLE_KINDA_SMALL_NUMBER));
+			CHECK_THAT(ResultPosition, Catch::ApproxEq(FVec3(-1, 2, 3)));
+			CHECK_THAT(ResultNormal, Catch::ApproxEq(FVec3(-1, 0, 0)));
+			CHECK_THAT(ResultFaceNormal, Catch::ApproxEq(FVec3(-1, 0, 0)));
+			CHECK(ResultFaceIndex == INDEX_NONE);
+		}
+		SECTION("Cap vs Edge 1")
+		{
+			// Sweep an x capsule so it hits the y capsule in the center
+			const FRigidTransform3 TestGeomTransform(FVec3(1, 2, 3), FRotation3::FromIdentity());
+			const FRigidTransform3 SweepGeomTransform(FVec3(-9, 2, 3), FRotation3::FromIdentity());
+			const FVec3 Dir(1, 0, 0);
+
+			const bool bResult = SweepQuery(YCapsule, TestGeomTransform, XCapsule, SweepGeomTransform, Dir, Length, ResultTime, ResultPosition, ResultNormal, ResultFaceIndex, ResultFaceNormal, Thickness, bComputeMTD);
+			CHECK(bResult);
+			CHECK_THAT(ResultTime, Catch::Matchers::WithinRel(5.0, UE_DOUBLE_KINDA_SMALL_NUMBER));
+			CHECK_THAT(ResultPosition, Catch::ApproxEq(FVec3(-1, 2, 3)));
+			CHECK_THAT(ResultNormal, Catch::ApproxEq(FVec3(-1, 0, 0)));
+			CHECK_THAT(ResultFaceNormal, Catch::ApproxEq(FVec3(-1, 0, 0)));
+			CHECK(ResultFaceIndex == INDEX_NONE);
+		}
+		SECTION("Cap vs Edge 2")
+		{
+			// Sweep a z capsule so its cap barely hits the y capsule in the center
+			const FRigidTransform3 TestGeomTransform(FVec3(1, 2, 3), FRotation3::FromIdentity());
+			const FRigidTransform3 SweepGeomTransform(FVec3(-9, 2, 0.1), FRotation3::FromIdentity());
+			const FVec3 Dir(1, 0, 0);
+
+			const bool bResult = SweepQuery(YCapsule, TestGeomTransform, ZCapsule, SweepGeomTransform, Dir, Length, ResultTime, ResultPosition, ResultNormal, ResultFaceIndex, ResultFaceNormal, Thickness, bComputeMTD);
+			CHECK(bResult);
+			CHECK_THAT(ResultTime, Catch::Matchers::WithinRel(6.48, 0.01));
+			CHECK_THAT(ResultPosition, Catch::ApproxEq(FVec3(-.759, 2, 2.049), 0.001));
+			CHECK_THAT(ResultNormal, Catch::ApproxEq(FVec3(-.880, 0, -.474), 0.001));
+			CHECK_THAT(ResultFaceNormal, Catch::ApproxEq(FVec3(-.880, 0, -.474), 0.001));
+			CHECK(ResultFaceIndex == INDEX_NONE);
+		}
+		SECTION("Cap vs Cap 1")
+		{
+			// Sweep an x capsule against another x capsule
+			const FRigidTransform3 TestGeomTransform(FVec3(1, 2, 3), FRotation3::FromIdentity());
+			const FRigidTransform3 SweepGeomTransform(FVec3(-9, 2, 3), FRotation3::FromIdentity());
+			const FVec3 Dir(1, 0, 0);
+
+			const bool bResult = SweepQuery(XCapsule, TestGeomTransform, XCapsule, SweepGeomTransform, Dir, Length, ResultTime, ResultPosition, ResultNormal, ResultFaceIndex, ResultFaceNormal, Thickness, bComputeMTD);
+			CHECK(bResult);
+			CHECK_THAT(ResultTime, Catch::Matchers::WithinRel(4, UE_DOUBLE_KINDA_SMALL_NUMBER));
+			CHECK_THAT(ResultPosition, Catch::ApproxEq(FVec3(-2, 2, 3)));
+			CHECK_THAT(ResultNormal, Catch::ApproxEq(FVec3(-1, 0, 0)));
+			CHECK_THAT(ResultFaceNormal, Catch::ApproxEq(FVec3(-1, 0, 0)));
+			CHECK(ResultFaceIndex == INDEX_NONE);
+		}
+		SECTION("Sweep: Past max distance")
+		{
+			const FRigidTransform3 TestGeomTransform(FVec3(1, 2, 3), FRotation3::FromIdentity());
+			const FRigidTransform3 SweepGeomTransform(FVec3(1 - Length - 6, 2, 3), FRotation3::FromIdentity());
+			const FVec3 Dir(1, 0, 0);
+
+			const bool bResult = SweepQuery(YCapsule, TestGeomTransform, XCapsule, SweepGeomTransform, Dir, Length, ResultTime, ResultPosition, ResultNormal, ResultFaceIndex, ResultFaceNormal, Thickness, bComputeMTD);
+			CHECK(!bResult);
+		}
+		SECTION("Sweep: Miss")
+		{
+			const FRigidTransform3 TestGeomTransform(FVec3(1, 2, 3), FRotation3::FromIdentity());
+			const FRigidTransform3 SweepGeomTransform(FVec3(-9, 10, 3), FRotation3::FromIdentity());
+			const FVec3 Dir(1, 0, 0);
+
+			const bool bResult = SweepQuery(YCapsule, TestGeomTransform, XCapsule, SweepGeomTransform, Dir, Length, ResultTime, ResultPosition, ResultNormal, ResultFaceIndex, ResultFaceNormal, Thickness, bComputeMTD);
+			CHECK(!bResult);
+		}
+		SECTION("Shallow Initial Overlap")
+		{
+			const FRigidTransform3 TestGeomTransform(FVec3(1, 2, 3), FRotation3::FromIdentity());
+			const FRigidTransform3 SweepGeomTransform(FVec3(-3.9, 2, 3), FRotation3::FromIdentity());
+			const FVec3 Dir(1, 0, 0);
+
+			const bool bResult = SweepQuery(YCapsule, TestGeomTransform, XCapsule, SweepGeomTransform, Dir, Length, ResultTime, ResultPosition, ResultNormal, ResultFaceIndex, ResultFaceNormal, Thickness, bComputeMTD);
+			CHECK(bResult);
+			CHECK_THAT(ResultTime, Catch::Matchers::WithinRel(-0.1, UE_DOUBLE_KINDA_SMALL_NUMBER));
+			CHECK_THAT(ResultPosition, Catch::ApproxEq(FVec3(-1, 2, 3)));
+			CHECK_THAT(ResultNormal, Catch::ApproxEq(FVec3(-1, 0, 0)));
+			CHECK_THAT(ResultFaceNormal, Catch::ApproxEq(FVec3(-1, 0, 0)));
+			CHECK(ResultFaceIndex == INDEX_NONE);
+		}
+		SECTION("Deep Initial Overlap")
+		{
+			const FRigidTransform3 TestGeomTransform(FVec3(1, 2, 3), FRotation3::FromIdentity());
+			const FRigidTransform3 SweepGeomTransform(FVec3(-0.1, 2, 3), FRotation3::FromIdentity());
+			const FVec3 Dir(1, 0, 0);
+
+			const bool bResult = SweepQuery(YCapsule, TestGeomTransform, XCapsule, SweepGeomTransform, Dir, Length, ResultTime, ResultPosition, ResultNormal, ResultFaceIndex, ResultFaceNormal, Thickness, bComputeMTD);
+			CHECK(bResult);
+			CHECK_THAT(ResultTime, Catch::Matchers::WithinRel(-3.9, UE_DOUBLE_KINDA_SMALL_NUMBER));
+			CHECK_THAT(ResultPosition, Catch::ApproxEq(FVec3(-1, 2, 3)));
+			CHECK_THAT(ResultNormal, Catch::ApproxEq(FVec3(-1, 0, 0)));
+			CHECK_THAT(ResultFaceNormal, Catch::ApproxEq(FVec3(-1, 0, 0)));
+			CHECK(ResultFaceIndex == INDEX_NONE);
+		}
+	}
+
+	TEST_CASE("SweepCapsuleVsCapsule - Perf", "[Chaos][Sweeps][!benchmark]")
+	{
+		const bool bComputeMTD = true;
+		const FReal Thickness = 0;
+		const FReal Length = 100;
+		const FCapsule XCapsule(FVec3(-1, 0, 0), FVec3(1, 0, 0), 2);
+		const FCapsule YCapsule(FVec3(0, -1, 0), FVec3(0, 1, 0), 2);
+		const FCapsule ZCapsule(FVec3(0, 0, -1), FVec3(0, 0, 1), 2);
+		const FRigidTransform3 TestGeomTransform(FVec3(1, 2, 3), FRotation3::FromIdentity());
+		const FVec3 Dir(1, 0, 0);
+
+		FReal ResultTime;
+		FVec3 ResultPosition;
+		FVec3 ResultNormal;
+		FVec3 ResultFaceNormal = FVec3::ZeroVector;
+		int32 ResultFaceIndex = INDEX_NONE;
+
+		BENCHMARK("Basic Hit")
+		{
+			const FRigidTransform3 SweepGeomTransform(FVec3(-9, 2, 3), FRotation3::FromIdentity());
+
+			bool bResult = false;
+			for (int32 I = 0; I < 1000; ++I)
+			{
+				bResult |= SweepQuery(XCapsule, TestGeomTransform, XCapsule, SweepGeomTransform, Dir, Length, ResultTime, ResultPosition, ResultNormal, ResultFaceIndex, ResultFaceNormal, Thickness, bComputeMTD);
+			}
+			return bResult;
+		};
+		BENCHMARK("Perpendicular Edge Edge")
+		{
+			const FRigidTransform3 SweepGeomTransform(FVec3(-9, 2, 3), FRotation3::FromIdentity());
+			const FVec3 Dir(1, 0, 0);
+
+			bool bResult = false;
+			for (int32 I = 0; I < 1000; ++I)
+			{
+				bResult |= SweepQuery(YCapsule, TestGeomTransform, ZCapsule, SweepGeomTransform, Dir, Length, ResultTime, ResultPosition, ResultNormal, ResultFaceIndex, ResultFaceNormal, Thickness, bComputeMTD);
+			}
+			return bResult;
+		};
+		BENCHMARK("Miss")
+		{
+			const FRigidTransform3 SweepGeomTransform(FVec3(-9, 10, 3), FRotation3::FromIdentity());
+
+			bool bResult = false;
+			for (int32 I = 0; I < 1000; ++I)
+			{
+				bResult |= SweepQuery(XCapsule, TestGeomTransform, XCapsule, SweepGeomTransform, Dir, Length, ResultTime, ResultPosition, ResultNormal, ResultFaceIndex, ResultFaceNormal, Thickness, bComputeMTD);
+			}
+			return bResult;
+		};
+		BENCHMARK("Shallow Initial Overlap")
+		{
+			const FRigidTransform3 SweepGeomTransform(FVec3(-3.9, 2, 3), FRotation3::FromIdentity());
+
+			bool bResult = false;
+			for (int32 I = 0; I < 1000; ++I)
+			{
+				bResult |= SweepQuery(YCapsule, TestGeomTransform, XCapsule, SweepGeomTransform, Dir, Length, ResultTime, ResultPosition, ResultNormal, ResultFaceIndex, ResultFaceNormal, Thickness, bComputeMTD);
+			}
+			return bResult;
+		};
+		BENCHMARK("Deep Initial Overlap")
+		{
+			const FRigidTransform3 SweepGeomTransform(FVec3(-0.1, 2, 3), FRotation3::FromIdentity());
+
+			bool bResult = false;
+			for (int32 I = 0; I < 1000; ++I)
+			{
+				bResult |= SweepQuery(XCapsule, TestGeomTransform, XCapsule, SweepGeomTransform, Dir, Length, ResultTime, ResultPosition, ResultNormal, ResultFaceIndex, ResultFaceNormal, Thickness, bComputeMTD);
+			}
+			return bResult;
+		};
+	}
+} // namespace Chaos
